ShaderVersion guid updated to reflect changes in the OpenGL shader compiler.

diff --git a/Engine/Shaders/Public/ShaderVersion.ush b/Engine/Shaders/Public/ShaderVersion.ush
--- a/Engine/Shaders/Public/ShaderVersion.ush
+++ b/Engine/Shaders/Public/ShaderVersion.ush
@@ -2,4 +2,4 @@
 // Each time the console command is executed it generates a new GUID. As this file is included
 // in Platform.ush (which should be included in any shader) it allows to invalidate the shader DDC.
 // 
-// GUID = 655DA074BADC40A1B8EFDEF3C906D9F0
+// GUID = 9D4644A1434A8C0C0DDDA5872682CB4E
\ No newline at end of file
