UE4 - Fix some potentially poisoned DDCs
#rnx

diff --git a/Engine/Shaders/Public/ShaderVersion.ush b/Engine/Shaders/Public/ShaderVersion.ush
--- a/Engine/Shaders/Public/ShaderVersion.ush
+++ b/Engine/Shaders/Public/ShaderVersion.ush
@@ -3,6 +3,7 @@
 // in Platform.ush (which should be included in any shader) it allows to invalidate the shader DDC.
 // 
 // If you are merging streams and there is a conflict with this GUID you should make a new GUID rather than taking one or the other.
-// GUID = 569B757A-6B67-4D56-AA92-6E8D8CD50EC3adwadefafafa
+// GUID = A9DE8FE2-CD23-4357-87D6-2B96A9F90453
+
 
 		  
