MN - Fix serialization of inline cooked materials

diff --git a/Engine/Shaders/Public/ShaderVersion.ush b/Engine/Shaders/Public/ShaderVersion.ush
--- a/Engine/Shaders/Public/ShaderVersion.ush
+++ b/Engine/Shaders/Public/ShaderVersion.ush
@@ -2,8 +2,7 @@
 // Each time the console command is executed it generates a new GUID. As this file is included
 // in Platform.ush (which should be included in any shader) it allows to invalidate the shader DDC.
 // 
-// If you are merging streams and there is a conflict with this GUID you should make a new GUID rather than taking one or the other.
-// GUID = A9DE8FE2-CD23-4357-87D6-2B96A9F90453
+// GUID = 1A20131D-6C53-4357-87D6-2B96A9F90453
 
 
 		  
diff --git a/Engine/Source/Runtime/RenderCore/Private/Shader.cpp b/Engine/Source/Runtime/RenderCore/Private/Shader.cpp
--- a/Engine/Source/Runtime/RenderCore/Private/Shader.cpp
+++ b/Engine/Source/Runtime/RenderCore/Private/Shader.cpp
@@ -645,7 +645,7 @@ void FShaderResource::Serialize(FArchive& Ar, bool bLoadedByCookedMaterial)
 	Ar << NumInstructions;
 
 #if WITH_EDITORONLY_DATA
-	if (!Ar.IsCooking() || Ar.CookingTarget()->HasEditorOnlyData())
+	if ((!Ar.IsCooking() || Ar.CookingTarget()->HasEditorOnlyData()) && !bLoadedByCookedMaterial)
 	{
 		Ar << NumTextureSamplers;
 	}
