Superior workaround for fast::sincos not being quite what we expect provided by Apple that should also workaround a driver bug on iOS.

diff --git a/Engine/Shaders/Public/Platform/Metal/MetalCommon.ush b/Engine/Shaders/Public/Platform/Metal/MetalCommon.ush
--- a/Engine/Shaders/Public/Platform/Metal/MetalCommon.ush
+++ b/Engine/Shaders/Public/Platform/Metal/MetalCommon.ush
@@ -5,7 +5,7 @@
 =============================================================================*/
 
 // Update this GUID to improve shader recompilation for Metal only shaders
-// GUID = 8D333468-602F-411C-A224-381A04952654
+// GUID = 547529DD-FBE7-45BC-A49B-0EFDB6D7F94E
   
 #pragma once
 
