Shaderversion needs a new GUID due to downstream issues

diff --git a/Engine/Shaders/Public/ShaderVersion.ush b/Engine/Shaders/Public/ShaderVersion.ush
--- a/Engine/Shaders/Public/ShaderVersion.ush
+++ b/Engine/Shaders/Public/ShaderVersion.ush
@@ -2,4 +2,4 @@
 // Each time the console command is executed it generates a new GUID. As this file is included
 // in Platform.ush (which should be included in any shader) it allows to invalidate the shader DDC.
 // 
-// GUID = 339F87C6-80DF-4A3F-8089-8AA5E9A73867
+// GUID = 487F1C22-EBE7-4CEA-A565-E19FE49FE8EE
\ No newline at end of file
