Put rules for shaderversion merging into the file

diff --git a/Engine/Shaders/Public/ShaderVersion.ush b/Engine/Shaders/Public/ShaderVersion.ush
--- a/Engine/Shaders/Public/ShaderVersion.ush
+++ b/Engine/Shaders/Public/ShaderVersion.ush
@@ -2,4 +2,5 @@
 // Each time the console command is executed it generates a new GUID. As this file is included
 // in Platform.ush (which should be included in any shader) it allows to invalidate the shader DDC.
 // 
+// If you are merging streams and there is a conflict with this GUID you should make a new GUID rather than taking one or the other.
 // GUID = 487F1C22-EBE7-4CEA-A565-E19FE49FE8EE
\ No newline at end of file
