DR - Fix Metal compilation (rebuilt hlslcc)

diff --git a/Engine/Shaders/Public/Platform/Metal/MetalCommon.ush b/Engine/Shaders/Public/Platform/Metal/MetalCommon.ush
--- a/Engine/Shaders/Public/Platform/Metal/MetalCommon.ush
+++ b/Engine/Shaders/Public/Platform/Metal/MetalCommon.ush
@@ -5,7 +5,7 @@
 =============================================================================*/
 
 // Update this GUID to improve shader recompilation for Metal only shaders
-// GUID = 0FDEF9A8-0978-44ED-B2ED-046924F18E72
+// GUID = 0FDEF9A8-0978-44ED-B2ED-046924F18E73
   
 #pragma once
 
diff --git a/Engine/Source/Developer/Apple/MetalShaderFormat/Private/MetalBackend.cpp b/Engine/Source/Developer/Apple/MetalShaderFormat/Private/MetalBackend.cpp
--- a/Engine/Source/Developer/Apple/MetalShaderFormat/Private/MetalBackend.cpp
+++ b/Engine/Source/Developer/Apple/MetalShaderFormat/Private/MetalBackend.cpp
@@ -1,5 +1,5 @@
 // Copyright 1998-2019 Epic Games, Inc. All Rights Reserved.
-// .....
+// .
 
 #include "MetalBackend.h"
 #include "MetalShaderFormat.h"
