Update GUID, ties to 8182408.

diff --git a/Engine/Shaders/Public/Platform/Metal/MetalCommon.ush b/Engine/Shaders/Public/Platform/Metal/MetalCommon.ush
--- a/Engine/Shaders/Public/Platform/Metal/MetalCommon.ush
+++ b/Engine/Shaders/Public/Platform/Metal/MetalCommon.ush
@@ -5,7 +5,7 @@
 =============================================================================*/
 
 // Update this GUID to improve shader recompilation for Metal only shaders
-// GUID = cda708c4-2f08-4eb2-b045-42a269d14ed4
+// GUID = DD2BA7D9-9401-42E4-975C-7812B6012971
 
 #pragma once
 
