Manually merge CL 8182418 from //UE4/Main/...
Transform access::write attribute for write-only UAVs.
#rnx

diff --git a/Engine/Shaders/Public/Platform/Metal/MetalCommon.ush b/Engine/Shaders/Public/Platform/Metal/MetalCommon.ush
--- a/Engine/Shaders/Public/Platform/Metal/MetalCommon.ush
+++ b/Engine/Shaders/Public/Platform/Metal/MetalCommon.ush
@@ -5,8 +5,8 @@
 =============================================================================*/
 
 // Update this GUID to improve shader recompilation for Metal only shaders
-// GUID = CC5D9FD4-2519-4E05-AC35-0317E9C937D4
-  
+// GUID = 7E37786E-C893-4EAF-8C7A-A200949A7438
+
 #pragma once
 
 #if VERTEXSHADER && HULLSHADER
