Fix MetalCommon.ush for generating ES3.1 shaders for iOS usign dxc.

diff --git a/Engine/Shaders/Public/Platform/Metal/MetalCommon.ush b/Engine/Shaders/Public/Platform/Metal/MetalCommon.ush
--- a/Engine/Shaders/Public/Platform/Metal/MetalCommon.ush
+++ b/Engine/Shaders/Public/Platform/Metal/MetalCommon.ush
@@ -30,7 +30,7 @@
 #define COMPILER_SUPPORTS_MINMAX3 0
 
 #if IOS
-#if MAX_SHADER_LANGUAGE_VERSION >= 6
+#if COMPILER_HLSLCC > 1
 #if PIXELSHADER
 	[[vk::input_attachment_index(0)]] SubpassInput<float4> gl_LastFragData;
 	#define FramebufferFetchES2() gl_LastFragData.SubpassLoad()
