Fixes missing shader crash in the AO denoiser.

diff --git a/Engine/Shaders/Public/ShaderVersion.ush b/Engine/Shaders/Public/ShaderVersion.ush
--- a/Engine/Shaders/Public/ShaderVersion.ush
+++ b/Engine/Shaders/Public/ShaderVersion.ush
@@ -2,5 +2,4 @@
 // Each time the console command is executed it generates a new GUID. As this file is included
 // in Platform.ush (which should be included in any shader) it allows to invalidate the shader DDC.
 // 
-// If you are merging streams and there is a conflict with this GUID you should make a new GUID rather than taking one or the other.
-// GUID = 490D5870-B9D3-4B88-B78C-DE1366151184
\ No newline at end of file
+// GUID = 1DEC0C1545F80AEC49E44F915EE0DEDB
\ No newline at end of file
diff --git a/Engine/Source/Runtime/Renderer/Private/ScreenSpaceDenoise.cpp b/Engine/Source/Runtime/Renderer/Private/ScreenSpaceDenoise.cpp
--- a/Engine/Source/Runtime/Renderer/Private/ScreenSpaceDenoise.cpp
+++ b/Engine/Source/Runtime/Renderer/Private/ScreenSpaceDenoise.cpp
@@ -223,6 +223,7 @@ static bool SignalSupportMultiSPP(ESignalProcessing SignalProcessing)
 	return (
 		SignalProcessing == ESignalProcessing::MonochromaticPenumbra ||
 		SignalProcessing == ESignalProcessing::Reflections ||
+		SignalProcessing == ESignalProcessing::AmbientOcclusion ||
 		SignalProcessing == ESignalProcessing::DiffuseIndirect);
 }
 
