- Implement INVARIANT on D3D to make use of the precise keyword
- This fixes various issues around depth divergence that can be see in the Niagara Mesh Vertex Factory as well as doing certain operations inside world position offset

diff --git a/Engine/Shaders/Public/Platform.ush b/Engine/Shaders/Public/Platform.ush
--- a/Engine/Shaders/Public/Platform.ush
+++ b/Engine/Shaders/Public/Platform.ush
@@ -35,6 +35,8 @@
 #elif COMPILER_VULKAN
 	// Helps with iteration when changing Vulkan shader code generation backend.
 	#include "Platform/Vulkan/VulkanCommon.ush"
+#elif SM5_PROFILE || SM4_PROFILE
+	#include "Platform/D3D/D3DCommon.ush"
 #endif
 // @todo platplug: Inject a platform's XXXCommon.ush externally, somehow
 
diff --git a/Engine/Shaders/Public/Platform/D3D/D3DCommon.ush b/Engine/Shaders/Public/Platform/D3D/D3DCommon.ush
new file mode 100644
--- /dev/null
+++ b/Engine/Shaders/Public/Platform/D3D/D3DCommon.ush
@@ -0,0 +1,11 @@
+// Copyright 1998-2019 Epic Games, Inc. All Rights Reserved.
+
+/*=============================================================================
+D3DCommon.usf: Common defines for the d3d shader compiler
+=============================================================================*/
+
+precise float MakePrecise(precise float v) { return v; }
+precise float2 MakePrecise(precise float2 v) { return v; }
+precise float3 MakePrecise(precise float3 v) { return v; }
+precise float4 MakePrecise(precise float4 v) { return v; }
+#define INVARIANT(X) MakePrecise(X)
diff --git a/Engine/Source/Developer/Windows/ShaderFormatD3D/Private/ShaderFormatD3D.cpp b/Engine/Source/Developer/Windows/ShaderFormatD3D/Private/ShaderFormatD3D.cpp
--- a/Engine/Source/Developer/Windows/ShaderFormatD3D/Private/ShaderFormatD3D.cpp
+++ b/Engine/Source/Developer/Windows/ShaderFormatD3D/Private/ShaderFormatD3D.cpp
@@ -72,7 +72,7 @@ public:
 	}
 	virtual const TCHAR* GetPlatformIncludeDirectory() const
 	{
-		return TEXT("");
+		return TEXT("D3D");
 	}
 };
 
