- Temporarily disable precise to fix engine tests
- This looks like a similar issue (driver / compiler) bug that happened with atmosphere as is fixed when running D3D12

diff --git a/Engine/Shaders/Public/Platform/D3D/D3DCommon.ush b/Engine/Shaders/Public/Platform/D3D/D3DCommon.ush
--- a/Engine/Shaders/Public/Platform/D3D/D3DCommon.ush
+++ b/Engine/Shaders/Public/Platform/D3D/D3DCommon.ush
@@ -8,4 +8,4 @@ precise float MakePrecise(precise float v) { return v; }
 precise float2 MakePrecise(precise float2 v) { return v; }
 precise float3 MakePrecise(precise float3 v) { return v; }
 precise float4 MakePrecise(precise float4 v) { return v; }
-#define INVARIANT(X) MakePrecise(X)
+#define INVARIANT(X) X
