Enable earlyZ

#rnx

diff --git a/Engine/Shaders/Public/Platform.ush b/Engine/Shaders/Public/Platform.ush
--- a/Engine/Shaders/Public/Platform.ush
+++ b/Engine/Shaders/Public/Platform.ush
@@ -474,7 +474,7 @@
 
 #endif // SM5_PROFILE || COMPILER_SUPPORTS_ATTRIBUTES
 
-#if SM5_PROFILE || METAL_MRT_PROFILE || METAL_SM5_PROFILE || METAL_SM5_NOTESS_PROFILE || ES3_1_PROFILE
+#if SM5_PROFILE || METAL_MRT_PROFILE || METAL_SM5_PROFILE || METAL_SM5_NOTESS_PROFILE || ES3_1_PROFILE || SWITCH_PROFILE
 #define EARLYDEPTHSTENCIL [earlydepthstencil]
 #endif
 
