Manually merge CL 10880209, and fix-up the GUID in MetalCommon.ush.

diff --git a/Engine/Shaders/Public/Platform/Metal/MetalCommon.ush b/Engine/Shaders/Public/Platform/Metal/MetalCommon.ush
--- a/Engine/Shaders/Public/Platform/Metal/MetalCommon.ush
+++ b/Engine/Shaders/Public/Platform/Metal/MetalCommon.ush
@@ -5,13 +5,7 @@
 =============================================================================*/
 
 // Update this GUID to improve shader recompilation for Metal only shaders
->>>> ORIGINAL //UE4/Main/Engine/Shaders/Public/Platform/Metal/MetalCommon.ush#20
-// GUID = 61EFF606-B5F2-42D6-9A48-9AD5E2E3BC55
-==== THEIRS //UE4/Main/Engine/Shaders/Public/Platform/Metal/MetalCommon.ush#22
-// GUID = 03B3D368-6D07-43AF-9EF6-3FAC517881D5
-==== YOURS //Juan.Canada_RPS/Engine/Shaders/Public/Platform/Metal/MetalCommon.ush
-// GUID = 616bd504-f38f-4c3a-83a8-fbb074511c78
-<<<<
+// GUID = D7A3E269-76AE-4BCD-939D-1A74146E1031
 
 #pragma once
 
diff --git a/Engine/Source/Runtime/Apple/MetalRHI/Private/MetalRHI.cpp b/Engine/Source/Runtime/Apple/MetalRHI/Private/MetalRHI.cpp
--- a/Engine/Source/Runtime/Apple/MetalRHI/Private/MetalRHI.cpp
+++ b/Engine/Source/Runtime/Apple/MetalRHI/Private/MetalRHI.cpp
@@ -120,8 +120,6 @@ FMetalDynamicRHI::FMetalDynamicRHI(ERHIFeatureLevel::Type RequestedFeatureLevel)
 	//@todo-rco: Query name from API
 	GRHIAdapterName = TEXT("Metal");
 	GRHIVendorId = 1; // non-zero to avoid asserts
-	
-	GRHISupportsMultithreading = true;
 
 	bool const bRequestedFeatureLevel = (RequestedFeatureLevel != ERHIFeatureLevel::Num);
 	bool bSupportsPointLights = false;
