Preserve half-precision floats by default for ES3_1_PROFILE.
See UDN thread: https://udn.unrealengine.com/questions/552871/half-precision-in-vulkan-shaders.html

#rnx

diff --git a/Engine/Shaders/Public/Platform.ush b/Engine/Shaders/Public/Platform.ush
--- a/Engine/Shaders/Public/Platform.ush
+++ b/Engine/Shaders/Public/Platform.ush
@@ -181,7 +181,7 @@
 #define FORCE_FLOATS 0
 #endif
 
-#if (!(COMPILER_GLSL_ES3_1 || METAL_PROFILE) || FORCE_FLOATS)
+#if (!(ES3_1_PROFILE || METAL_PROFILE) || FORCE_FLOATS)
 	// Always use floats when using the ES3/METAL compiler, because platforms not optimized for lower precision, 
 	// And we don't want potential side effects on other platforms
 	#define half float
