Bumping the shader version on chance that prior change11706862  might necessitate it. Made using r.InvalidateShaderCache

diff --git a/Engine/Shaders/Public/ShaderVersion.ush b/Engine/Shaders/Public/ShaderVersion.ush
--- a/Engine/Shaders/Public/ShaderVersion.ush
+++ b/Engine/Shaders/Public/ShaderVersion.ush
@@ -3,6 +3,4 @@
 // in Platform.ush (which should be included in any shader) it allows to invalidate the shader DDC.
 // 
 // If you are merging streams and there is a conflict with this GUID you should make a new GUID rather than taking one or the other.
-// GUID = da9004fb-2012-49c7-aabd-f0446c4aa3a1
-
-
+// GUID = 979F941B400ECB2A9CDA3CB0FBA2FC9E
\ No newline at end of file
