Bump MetalCommon.ush GUID.

#rnx

diff --git a/Engine/Shaders/Public/Platform/Metal/MetalCommon.ush b/Engine/Shaders/Public/Platform/Metal/MetalCommon.ush
--- a/Engine/Shaders/Public/Platform/Metal/MetalCommon.ush
+++ b/Engine/Shaders/Public/Platform/Metal/MetalCommon.ush
@@ -5,7 +5,7 @@
 =============================================================================*/
 
 // Update this GUID to improve shader recompilation for Metal only shaders
-// GUID = D8D6C149-3862-4960-A60F-DD7862433434
+// GUID = E3AC9C39-C9CD-48C1-8F24-DEDF572A4871
 
 #pragma once
 
