- Moved some platform shader code to PlatformExtensions

diff --git a/Engine/Shaders/Public/Platform.ush b/Engine/Shaders/Public/Platform.ush
--- a/Engine/Shaders/Public/Platform.ush
+++ b/Engine/Shaders/Public/Platform.ush
@@ -22,11 +22,6 @@
 
 #ifdef OVERRIDE_PLATFORMCOMMON_USH
 	#include "/Platform/Public/PlatformCommon.ush"
-#elif PS4_PROFILE  
-	// always #include PS4Common.usf so it can #define override anything in any other included file.
-	#include "Platform/PS4/PS4Common.ush"
-#elif XBOXONE_PROFILE
-	#include "Platform/XboxOne/XboxOneCommon.ush"
 #elif COMPILER_METAL
 	// Helps with iteration when changing Metal shader code generation backend.
 	#include "Platform/Metal/MetalCommon.ush"
