Removed COMPILE_MODE_WAVE_32 and COMPILE_MODE_WAVE_64 defines from all platforms, in favor of CFLAG_Wave32

diff --git a/Engine/Shaders/Public/Platform.ush b/Engine/Shaders/Public/Platform.ush
--- a/Engine/Shaders/Public/Platform.ush
+++ b/Engine/Shaders/Public/Platform.ush
@@ -176,8 +176,6 @@
 
 #ifndef COMPILER_SUPPORTS_WAVE_32_64_MODE
 #define COMPILER_SUPPORTS_WAVE_32_64_MODE 0
-#define COMPILE_MODE_WAVE_32
-#define COMPILE_MODE_WAVE_64
 #endif
 
 // ---------------------------------------------------- Alternative floating point types
