Update MetalCommon.ush GUID

- Follow-on to previous commit.

#rnx

diff --git a/Engine/Shaders/Public/Platform/Metal/MetalCommon.ush b/Engine/Shaders/Public/Platform/Metal/MetalCommon.ush
--- a/Engine/Shaders/Public/Platform/Metal/MetalCommon.ush
+++ b/Engine/Shaders/Public/Platform/Metal/MetalCommon.ush
@@ -5,7 +5,7 @@
 =============================================================================*/
 
 // Update this GUID to improve shader recompilation for Metal only shaders
-// GUID = 7742EA46-5497-4ABC-897F-132501941942
+// GUID = 81F5B3FA-F43E-4A5E-B3FD-A466284A199C
 
 #pragma once
 
