Reduces the alignment of FRDG{Texture,Buffer}Access and allow casting to pointer

#[fyi] zach.bethel

diff --git a/Engine/Shaders/Public/ShaderVersion.ush b/Engine/Shaders/Public/ShaderVersion.ush
--- a/Engine/Shaders/Public/ShaderVersion.ush
+++ b/Engine/Shaders/Public/ShaderVersion.ush
@@ -2,4 +2,4 @@
 // Each time the console command is executed it generates a new GUID. As this file is included
 // in Platform.ush (which should be included in any shader) it allows to invalidate the shader DDC.
 // 
-// GUID = F4FB83CC4C6472BBB80D0EA07AE3A02F
\ No newline at end of file
+// GUID = F4FB83CC4C6472DDD80D0EA07AE3A02F
\ No newline at end of file
