Added shader compiler support for [noinline] where available. Informs compiler we want a subroutine created, which can be used to decrease register pressure in certain situations. Code is kept separate, and a set number of registers are used on each call. Should only be used with extensive profiling, as the default inlining behavior is usually best.

DXIL:  [REDACTED_PROFILE_URL]/DirectXShaderCompiler/blob/master/tools/clang/test/HLSLFileCheck/hlsl/functions/attribute/noinline.hlsl
SPIRV: https://www.khronos.org/registry/spir-v/specs/unified1/SPIRV.html (DontInline)

Shader code can now use NOINLINE, and also test if it maps to something useful with COMPILER_SUPPORTS_NOINLINE

diff --git a/Engine/Shaders/Public/Platform.ush b/Engine/Shaders/Public/Platform.ush
--- a/Engine/Shaders/Public/Platform.ush
+++ b/Engine/Shaders/Public/Platform.ush
@@ -507,18 +507,34 @@ void ShaderYield()
 #define FLATTEN
 #endif
 
-#ifndef ALLOW_UAV_CONDITION	
-#define ALLOW_UAV_CONDITION	
+#ifndef ALLOW_UAV_CONDITION
+#define ALLOW_UAV_CONDITION
 #endif
 
 #ifndef INVARIANT
-#define INVARIANT										
+#define INVARIANT
 #endif
 
 #ifndef ENABLE_RE_Z
 #define ENABLE_RE_Z
 #endif
 
+#ifndef COMPILER_SUPPORTS_NOINLINE
+#define COMPILER_SUPPORTS_NOINLINE 0
+#endif
+
+// Informs compiler we want a subroutine created, which can be used to 
+// decrease register pressure in certain situations. Code is kept separate,
+// and a set number of registers are used on each call. Should only be used
+// with extensive profiling, as the default inlining behavior is usually best.
+// DXIL:  [REDACTED_PROFILE_URL]/DirectXShaderCompiler/blob/master/tools/clang/test/HLSLFileCheck/hlsl/functions/attribute/noinline.hlsl
+// SPIRV: https://www.khronos.org/registry/spir-v/specs/unified1/SPIRV.html (DontInline)
+#if COMPILER_SUPPORTS_NOINLINE
+#define NOINLINE [noinline]
+#else
+#define NOINLINE
+#endif
+
 #ifndef EARLYDEPTHSTENCIL
 #define EARLYDEPTHSTENCIL
 #endif
diff --git a/Engine/Source/Developer/Windows/ShaderFormatD3D/Private/D3DShaderCompiler.cpp b/Engine/Source/Developer/Windows/ShaderFormatD3D/Private/D3DShaderCompiler.cpp
--- a/Engine/Source/Developer/Windows/ShaderFormatD3D/Private/D3DShaderCompiler.cpp
+++ b/Engine/Source/Developer/Windows/ShaderFormatD3D/Private/D3DShaderCompiler.cpp
@@ -914,6 +914,7 @@ void CompileD3DShader(const FShaderCompilerInput& Input, FShaderCompilerOutput&
 	{
 		AdditionalDefines.SetDefine(TEXT("PLATFORM_SUPPORTS_SM6_0_WAVE_OPERATIONS"), 1);
 		AdditionalDefines.SetDefine(TEXT("PLATFORM_SUPPORTS_STATIC_SAMPLERS"), 1);
+		AdditionalDefines.SetDefine(TEXT("COMPILER_SUPPORTS_NOINLINE"), 1);
 	}
 
 	if (Input.bSkipPreprocessedCache)
