Cleanup: removing uses of the outdated SM4_PROFILE

diff --git a/Engine/Shaders/Private/BasePassPixelShader.usf b/Engine/Shaders/Private/BasePassPixelShader.usf
--- a/Engine/Shaders/Private/BasePassPixelShader.usf
+++ b/Engine/Shaders/Private/BasePassPixelShader.usf
@@ -1443,7 +1443,7 @@ void FPixelShaderInOut_MainPS(
 	// The following block is disabled on Vulkan because it triggers an nvidia driver bug (UE-101609).
 #if USE_DEVELOPMENT_SHADERS && !VULKAN_PROFILE_SM5
 	// this feature is only needed for development/editor - we can compile it out for a shipping build (see r.CompileShadersForDevelopment cvar help)
-	#if METAL_SM5_PROFILE || SM6_PROFILE || SM5_PROFILE || SM4_PROFILE || VULKAN_PROFILE_SM5
+	#if METAL_SM5_PROFILE || SM6_PROFILE || SM5_PROFILE || VULKAN_PROFILE_SM5
 		BRANCH
 		if (View.OutOfBoundsMask > 0)
 		{
diff --git a/Engine/Shaders/Public/Platform.ush b/Engine/Shaders/Public/Platform.ush
--- a/Engine/Shaders/Public/Platform.ush
+++ b/Engine/Shaders/Public/Platform.ush
@@ -29,7 +29,7 @@
 #elif COMPILER_VULKAN
 	// Helps with iteration when changing Vulkan shader code generation backend.
 	#include "Platform/Vulkan/VulkanCommon.ush"
-#elif SM6_PROFILE || SM5_PROFILE || SM4_PROFILE
+#elif SM6_PROFILE || SM5_PROFILE
 	#include "Platform/D3D/D3DCommon.ush"
 #endif
 
@@ -93,10 +93,6 @@
 #define SM5_PROFILE 0
 #endif
 
-#ifndef SM4_PROFILE
-#define SM4_PROFILE 0
-#endif
-
 #ifndef OPENGL_PROFILE
 #define OPENGL_PROFILE 0
 #endif
@@ -213,9 +209,6 @@
 	// SM5 = full dx11 features (high end UE4 rendering)
 	#define FEATURE_LEVEL FEATURE_LEVEL_SM5
 
-#elif SM4_PROFILE
-	#define FEATURE_LEVEL FEATURE_LEVEL_SM4
-
 #elif SWITCH_PROFILE || SWITCH_PROFILE_FORWARD
 	#undef ES3_1_PROFILE
 
@@ -334,13 +327,13 @@
 // Ideally this would just be enabled for all SM4 and above feature level platforms
 // @todo metal mrt: No reason this can't work with Metal, once cross compiler is fixed
 #ifndef SUPPORTS_INDEPENDENT_SAMPLERS
-#define SUPPORTS_INDEPENDENT_SAMPLERS (SM6_PROFILE || SM5_PROFILE || SM4_PROFILE || METAL_MRT_PROFILE || METAL_SM5_PROFILE || VULKAN_PROFILE_SM5 || VULKAN_PROFILE)
+#define SUPPORTS_INDEPENDENT_SAMPLERS (SM6_PROFILE || SM5_PROFILE || METAL_MRT_PROFILE || METAL_SM5_PROFILE || VULKAN_PROFILE_SM5 || VULKAN_PROFILE)
 #endif
 
 // Whether the platform supports a global clip plane through SV_ClipDistance
 // Ideally this would just be enabled for all SM4 and above feature level platforms, but not tested everywhere yet
 #ifndef PLATFORM_SUPPORTS_GLOBAL_CLIP_PLANE
-#define PLATFORM_SUPPORTS_GLOBAL_CLIP_PLANE (PS4_PROFILE || SM6_PROFILE || SM5_PROFILE || SM4_PROFILE || METAL_PROFILE || METAL_MRT_PROFILE || METAL_SM5_PROFILE || VULKAN_PROFILE_SM5)
+#define PLATFORM_SUPPORTS_GLOBAL_CLIP_PLANE (PS4_PROFILE || SM6_PROFILE || SM5_PROFILE || METAL_PROFILE || METAL_MRT_PROFILE || METAL_SM5_PROFILE || VULKAN_PROFILE_SM5)
 #endif
 
 // Whether the platform support pixel coverage on MSAA targets (SV_Coverage).
