Add generic empty PLATFORM_BREAK definition to Platform.ush

diff --git a/Engine/Shaders/Public/Platform.ush b/Engine/Shaders/Public/Platform.ush
--- a/Engine/Shaders/Public/Platform.ush
+++ b/Engine/Shaders/Public/Platform.ush
@@ -348,6 +348,11 @@
 
 // ---------------------------------------------------- Compiler specific defaults and fallbacks
 
+#if !defined(PLATFORM_BREAK)
+	#define PLATFORM_BREAK() 
+#endif
+
+
 // Hlslcc platforms ignore the uniform keyword as it can't properly optimize flow
 #if COMPILER_HLSLCC
 #define uniform
