Invalidate Vulkan shaders as appendix to DXC update in CL 16186169.

#rnx

diff --git a/Engine/Shaders/Public/Platform/Vulkan/VulkanCommon.ush b/Engine/Shaders/Public/Platform/Vulkan/VulkanCommon.ush
--- a/Engine/Shaders/Public/Platform/Vulkan/VulkanCommon.ush
+++ b/Engine/Shaders/Public/Platform/Vulkan/VulkanCommon.ush
@@ -5,6 +5,6 @@
 =============================================================================*/
 
 // Update this GUID to improve shader recompilation for Vulkan only shaders
-// GUID = 42743725-3246-45DA-8795-422991BBBF1E
+// GUID = 9434126B-7630-4238-BE07-BF4E9C5AF8DB
 
 #pragma once
