Adding cross platform shader support for Mesh shaders.

diff --git a/Engine/Shaders/Public/Platform/D3D/D3DCommon.ush b/Engine/Shaders/Public/Platform/D3D/D3DCommon.ush
--- a/Engine/Shaders/Public/Platform/D3D/D3DCommon.ush
+++ b/Engine/Shaders/Public/Platform/D3D/D3DCommon.ush
@@ -62,3 +62,20 @@ void UEReportAssertWithPayload(uint AssertID, uint4 Payload)
 #define PLATFORM_ASSERT4(condition, assert_id, a, b, c, d) { if (!(condition)) UEReportAssertWithPayload(assert_id, uint4(a, b, c, d));  }
 
 #endif // PLATFORM_SUPPORTS_DIAGNOSTIC_BUFFER
+
+#if PLATFORM_SUPPORTS_MESH_SHADERS
+	#define MESH_SHADER_TRIANGLE_ATTRIBUTES(InThreadsX) [numthreads(InThreadsX, 1, 1)][outputtopology("triangle")]
+	#define MESH_SHADER_LINE_ATTRIBUTES(InThreadsX) [numthreads(InThreadsX, 1, 1)][outputtopology("line")]
+
+	#define AMPLIFICATION_SHADER_TRIANGLE_ATTRIBUTES(InThreadsX) [numthreads(InThreadsX, 1, 1)][outputtopology("triangle")]
+	#define AMPLIFICATION_SHADER_LINE_ATTRIBUTES(InThreadsX) [numthreads(InThreadsX, 1, 1)][outputtopology("line")]
+
+	#define MESH_SHADER_VERTEX_EXPORT(VertexType, InMaxVertexCount) out vertices VertexType OutVertices[InMaxVertexCount]
+	#define MESH_SHADER_TRIANGLE_EXPORT(InMaxTriangleCount) out indices uint3 OutTriangles[InMaxTriangleCount]
+	#define MESH_SHADER_LINE_EXPORT(InMaxLineCount) out indices uint2 OutLines[InMaxLineCount]
+
+	#define MESH_SHADER_WRITE_VERTEX(InVertexIndex, Value) OutVertices[InVertexIndex] = Value;
+	#define MESH_SHADER_WRITE_TRIANGLE(InTriangleIndex, Value) OutTriangles[InTriangleIndex] = Value;
+	#define MESH_SHADER_WRITE_LINE(InLineIndex, Value) OutLines[InLineIndex] = Value;
+#endif // PLATFORM_SUPPORTS_MESH_SHADERS
+
diff --git a/Engine/Source/Runtime/Engine/Private/ShaderCompiler/ShaderCompiler.cpp b/Engine/Source/Runtime/Engine/Private/ShaderCompiler/ShaderCompiler.cpp
--- a/Engine/Source/Runtime/Engine/Private/ShaderCompiler/ShaderCompiler.cpp
+++ b/Engine/Source/Runtime/Engine/Private/ShaderCompiler/ShaderCompiler.cpp
@@ -5219,6 +5219,7 @@ void GlobalBeginCompileShader(
 	Input.Environment.SetDefine(TEXT("PLATFORM_SUPPORTS_RENDERTARGET_WRITE_MASK"), RHISupportsRenderTargetWriteMask(EShaderPlatform(Target.Platform)) ? 1 : 0);
 	Input.Environment.SetDefine(TEXT("PLATFORM_SUPPORTS_PER_PIXEL_DBUFFER_MASK"), IsUsingPerPixelDBufferMask(EShaderPlatform(Target.Platform)) ? 1 : 0);
 	Input.Environment.SetDefine(TEXT("PLATFORM_SUPPORTS_DISTANCE_FIELDS"), DoesPlatformSupportDistanceFields(EShaderPlatform(Target.Platform)) ? 1 : 0);
+	Input.Environment.SetDefine(TEXT("PLATFORM_SUPPORTS_MESH_SHADERS"), RHISupportsMeshShaders(EShaderPlatform(Target.Platform)) ? 1 : 0);
 	
 	{
 		static const auto CVar = IConsoleManager::Get().FindConsoleVariable(TEXT("r.vt.FeedbackFactor"));
diff --git a/Engine/Source/Runtime/RenderCore/Private/Shader.cpp b/Engine/Source/Runtime/RenderCore/Private/Shader.cpp
--- a/Engine/Source/Runtime/RenderCore/Private/Shader.cpp
+++ b/Engine/Source/Runtime/RenderCore/Private/Shader.cpp
@@ -1800,6 +1800,11 @@ void ShaderMapAppendKeyString(EShaderPlatform Platform, FString& KeyString)
 		KeyString += TEXT("_PPDBM");
 	}
 
+	if (RHISupportsMeshShaders(Platform))
+	{
+		KeyString += TEXT("_MS");
+	}
+
 	if (ShouldCompileRayTracingShadersForProject(Platform))
 	{
 		static const auto CVarCompileCHS = IConsoleManager::Get().FindConsoleVariable(TEXT("r.RayTracing.CompileMaterialCHS"));
