Added support for high frequency timestamp queries in shaders

diff --git a/Engine/Shaders/Public/Platform.ush b/Engine/Shaders/Public/Platform.ush
--- a/Engine/Shaders/Public/Platform.ush
+++ b/Engine/Shaders/Public/Platform.ush
@@ -360,6 +360,10 @@
 	#define PLATFORM_ASSERT4(condition, assert_id, a, b, c, d)
 #endif
 
+#if !defined(PLATFORM_SUPPORTS_SHADER_TIMESTAMP)
+	#define PLATFORM_SUPPORTS_SHADER_TIMESTAMP 0
+#endif
+
 // Hlslcc platforms ignore the uniform keyword as it can't properly optimize flow
 #if COMPILER_HLSLCC
 #define uniform
