Bump VulkanCommon.ush GUID due to CL 16346836 changing format of the shaders data

diff --git a/Engine/Shaders/Public/Platform/Vulkan/VulkanCommon.ush b/Engine/Shaders/Public/Platform/Vulkan/VulkanCommon.ush
--- a/Engine/Shaders/Public/Platform/Vulkan/VulkanCommon.ush
+++ b/Engine/Shaders/Public/Platform/Vulkan/VulkanCommon.ush
@@ -5,6 +5,6 @@
 =============================================================================*/
 
 // Update this GUID to improve shader recompilation for Vulkan only shaders
-// GUID = 9434126B-7630-4238-BE07-BF4E9C5AF8DB
+// GUID = E28F6421-6CC3-4D01-BCED-D8442AEF87A9
 
 #pragma once
