Fix attemp for material with anisotropy not being registered runtime side. It seems the code is right, a shader version bump fixes it (by forcing cache update).

diff --git a/Engine/Shaders/Public/ShaderVersion.ush b/Engine/Shaders/Public/ShaderVersion.ush
--- a/Engine/Shaders/Public/ShaderVersion.ush
+++ b/Engine/Shaders/Public/ShaderVersion.ush
@@ -3,4 +3,4 @@
 // in Platform.ush (which should be included in any shader) it allows to invalidate the shader DDC.
 // 
 // If you are merging streams and there is a conflict with this GUID you should make a new GUID rather than taking one or the other.
-// GUID = FF18584666CC4406A481AB221199BBAA
+// GUID = BF6FC041379A4FA6A080522A47CDB403
