Fixes some of the dual pixel vectorization helpers being not available for 16bit integers

diff --git a/Engine/Shaders/Private/TemporalSuperResolution/TSRUpdateHistory.usf b/Engine/Shaders/Private/TemporalSuperResolution/TSRUpdateHistory.usf
--- a/Engine/Shaders/Private/TemporalSuperResolution/TSRUpdateHistory.usf
+++ b/Engine/Shaders/Private/TemporalSuperResolution/TSRUpdateHistory.usf
@@ -948,13 +948,7 @@ void MainCS(
 			}
 		}
 		
-		#if COMPILER_DXC == 1 // shader compiler bug
-			uint2x2 bKOSignInt = uint2x2(bKOSign);
-			uint2 KOSignVector = uint2(1, SUB_PIXEL_GRID_SIZE);
-			uint2 SubpixelId = dpv_interleave_registers(dot(dpv_lo(bKOSignInt), KOSignVector), dot(dpv_hi(bKOSignInt), KOSignVector));
-		#else
-			uint2 SubpixelId = dpv_dot(uint2x2(bKOSign), uint2(1, SUB_PIXEL_GRID_SIZE));
-		#endif
+		uint2 SubpixelId = dpv_dot(uint2x2(bKOSign), uint2(1, SUB_PIXEL_GRID_SIZE));
 
 		taa_subpixel_payload2 SubpixelShift = taa_subpixel_payload2(SubpixelId) * taa_subpixel_payload(SUB_PIXEL_BIT_COUNT);
 
diff --git a/Engine/Shaders/Public/DualPixelVectorization.ush b/Engine/Shaders/Public/DualPixelVectorization.ush
--- a/Engine/Shaders/Public/DualPixelVectorization.ush
+++ b/Engine/Shaders/Public/DualPixelVectorization.ush
@@ -293,15 +293,6 @@ DEFINE_DPV_FUNCTION_FOR_ALL_VECTORS(DEFINE_DPV_FUNCTION_VECTOR)
 	__DPV_MULTI_SCALAR(TYPE) dpv_dot(__DPV_MULTI_VECTOR(TYPE, COMPONENTS) a, __DPV_VECTOR(TYPE, COMPONENTS) b) { \
 		return dpv_interleave_registers(dot(dpv_lo(a), b), dot(dpv_hi(a), b)); \
 	} \
-	__DPV_MULTI_SCALAR(TYPE) dpv_length2(__DPV_MULTI_VECTOR(TYPE, COMPONENTS) v) { \
-		return dpv_dot(v, v); \
-	} \
-	__DPV_MULTI_SCALAR(TYPE) dpv_length(__DPV_MULTI_VECTOR(TYPE, COMPONENTS) v) { \
-		return sqrt(dpv_length2(v)); \
-	} \
-	__DPV_MULTI_VECTOR(TYPE, COMPONENTS) dpv_normalize(__DPV_MULTI_VECTOR(TYPE, COMPONENTS) v) { \
-		return dpv_scale(v, rsqrt(dpv_length2(v))); \
-	} \
 	__DPV_MULTI_VECTOR(TYPE, COMPONENTS) dpv_clamp(__DPV_MULTI_VECTOR(TYPE, COMPONENTS) x, __DPV_MULTI_VECTOR(TYPE, COMPONENTS) min, __DPV_MULTI_VECTOR(TYPE, COMPONENTS) max) { \
 		__DPV_MULTI_VECTOR(TYPE, COMPONENTS) r; \
 		UNROLL_N(COMPONENTS) for (uint c = 0; c < COMPONENTS; c++) { \
@@ -344,6 +335,28 @@ DEFINE_DPV_FUNCTION_FOR_ALL_VECTORS(DEFINE_DPV_FUNCTION_VECTOR)
 		return r; \
 	} \
 
+DEFINE_DPV_FUNCTION_FOR_ALL_VECTORS(DEFINE_DPV_FUNCTION_VECTOR)
+#undef DEFINE_DPV_FUNCTION_VECTOR
+
+
+/* per vector vectorial operations
+ *
+ * Example:
+ *		half3x2 WorldPositionDelta = ...;
+ * 
+ *		half2 Distance = dpv_length(WorldPositionDelta);
+ */
+#define DEFINE_DPV_FUNCTION_VECTOR(TYPE, COMPONENTS) \
+	__DPV_MULTI_SCALAR(TYPE) dpv_length2(__DPV_MULTI_VECTOR(TYPE, COMPONENTS) v) { \
+		return dpv_dot(v, v); \
+	} \
+	__DPV_MULTI_SCALAR(TYPE) dpv_length(__DPV_MULTI_VECTOR(TYPE, COMPONENTS) v) { \
+		return sqrt(dpv_length2(v)); \
+	} \
+	__DPV_MULTI_VECTOR(TYPE, COMPONENTS) dpv_normalize(__DPV_MULTI_VECTOR(TYPE, COMPONENTS) v) { \
+		return dpv_scale(v, rsqrt(dpv_length2(v))); \
+	} \
+
 
 DEFINE_DPV_FUNCTION_FOR_FLOAT_VECTORS(DEFINE_DPV_FUNCTION_VECTOR)
 #undef DEFINE_DPV_FUNCTION_VECTOR
