Adding RHIUniformBufferLayoutInitializer and RHIUniformBufferLayout is a proper FRHIResource. This brings proper RHI lifetime support to RHIUniformBufferLayout so we don't delete the layouts while RHI commands that use them are still in flight.

This should fully address UniformBufferLayout/UniformBuffer crashes that randomly happen during shader compiling in the Editor.

diff --git a/Engine/Shaders/Public/ShaderVersion.ush b/Engine/Shaders/Public/ShaderVersion.ush
--- a/Engine/Shaders/Public/ShaderVersion.ush
+++ b/Engine/Shaders/Public/ShaderVersion.ush
@@ -3,4 +3,4 @@
 // in Platform.ush (which should be included in any shader) it allows to invalidate the shader DDC.
 // 
 // If you are merging streams and there is a conflict with this GUID you should make a new GUID rather than taking one or the other.
-// GUID = 5F9FBB1B163644EA9F41E0B80737816F
+// GUID = E5401FDD59AF47489260A452B32036EB
