New byte alignment function ByteAlignU32()

- default implementation provided
- can be specialized per-platform by defining COMPILER_SUPPORTS_BYTEALIGN

diff --git a/Engine/Shaders/Public/Platform.ush b/Engine/Shaders/Public/Platform.ush
--- a/Engine/Shaders/Public/Platform.ush
+++ b/Engine/Shaders/Public/Platform.ush
@@ -830,6 +830,14 @@ uint BitAlignU32(uint High, uint Low, uint Shift)
 }
 #endif
 
+#ifndef COMPILER_SUPPORTS_BYTEALIGN
+#define COMPILER_SUPPORTS_BYTEALIGN 0
+uint ByteAlignU32(uint High, uint Low, uint Shift)
+{
+	return BitAlignU32(High, Low, Shift * 8);
+}
+#endif // #ifndef COMPILER_SUPPORTS_BYTEALIGN
+
 #if COMPILER_HLSLCC
 #define ddx_fine(x) ddx(x)
 #define ddy_fine(y) ddy(y)
