Fixed bug with Vulkan not using half precision

diff --git a/Engine/Shaders/Public/Platform.ush b/Engine/Shaders/Public/Platform.ush
--- a/Engine/Shaders/Public/Platform.ush
+++ b/Engine/Shaders/Public/Platform.ush
@@ -207,7 +207,7 @@
 	#define fixed3x3 float3x3
 	#define fixed4x4 float4x4
 	#define fixed4x3 float4x3
-#elif (VULKAN_PROFILE || COMPILER_GLSL_ES3_1) && !COMPILER_HLSLCC
+#elif (VULKAN_PROFILE) || (COMPILER_GLSL_ES3_1 && !(COMPILER_HLSLCC && COMPILER_HLSLCC == 1))
 	// For VULKAN and OPENGL ES31 use RelaxedPrecision for half floats
 	#define half min16float
 	#define half2 min16float2
