Implements CALL_SITE_DEBUGLOC

diff --git a/Engine/Shaders/Private/TemporalSuperResolution/TSRCommon.ush b/Engine/Shaders/Private/TemporalSuperResolution/TSRCommon.ush
--- a/Engine/Shaders/Private/TemporalSuperResolution/TSRCommon.ush
+++ b/Engine/Shaders/Private/TemporalSuperResolution/TSRCommon.ush
@@ -285,6 +285,7 @@ uint bCameraCut;
 
 //------------------------------------------------------- FUNCTIONS
 	
+CALL_SITE_DEBUGLOC 
 float2x2 ApplyScreenTransform(float2x2 PInA, FScreenTransform AToB)
 {
 	return dpv_add(dpv_mul(PInA, AToB.xy), AToB.zw);
@@ -292,6 +293,7 @@ float2x2 ApplyScreenTransform(float2x2 PInA, FScreenTransform AToB)
 
 #if PLATFORM_SUPPORTS_REAL_TYPES
 
+CALL_SITE_DEBUGLOC 
 half2x2 ApplyScreenTransform(half2x2 PInA, FScreenTransform AToB)
 {
 	return dpv_add(dpv_mul(PInA, half2(AToB.xy)), half2(AToB.zw));
@@ -301,6 +303,7 @@ half2x2 ApplyScreenTransform(half2x2 PInA, FScreenTransform AToB)
 
 
 /** Compute the group wave index into SGRP to then recompue the GroupThreadIndex latter */
+CALL_SITE_DEBUGLOC 
 uint GetGroupWaveIndex(uint GroupThreadIndex, uint GroupSize)
 #if COMPILER_SUPPORTS_WAVE_ONCE
 {
@@ -320,6 +323,7 @@ uint GetGroupWaveIndex(uint GroupThreadIndex, uint GroupSize)
 #endif
 
 /** Force compute the group GroupThreadIndex through lane index and wave index if possible to reduce VGPR pressure. */
+CALL_SITE_DEBUGLOC 
 uint GetGroupThreadIndex(uint GroupThreadIndex, uint GroupWaveIndex)
 #if COMPILER_SUPPORTS_WAVE_ONCE
 {
@@ -339,6 +343,7 @@ uint GetGroupThreadIndex(uint GroupThreadIndex, uint GroupWaveIndex)
 #endif
 
 // Clamp the offset to be shared across multiple samples
+CALL_SITE_DEBUGLOC
 taa_short2x2 ClampPixelOffset(
 	taa_short2x2 KernelCenterPixelPos,
 	taa_short2x2 Offset,
@@ -373,6 +378,7 @@ taa_short2x2 ClampPixelOffset(
 	return ClampedOffset;
 }
 
+CALL_SITE_DEBUGLOC
 taa_short2 ClampPixelOffset(
 	taa_short2 KernelCenterPixelPos,
 	taa_short2 Offset,
@@ -399,6 +405,7 @@ taa_short2 ClampPixelOffset(
 	return ClampedOffset;
 }
 
+CALL_SITE_DEBUGLOC
 taa_short2x2 ClampPixelOffset(
 	taa_short2x2 SamplePixelPos,
 	int2 MinPixelPos, int2 MaxPixelPos)
@@ -421,6 +428,7 @@ taa_short2x2 ClampPixelOffset(
 	return dpv_interleave_registers(SamplePixelPos0, SamplePixelPos1);
 }
 
+CALL_SITE_DEBUGLOC
 taa_short2 ClampPixelOffset(
 	taa_short2 SamplePixelPos,
 	int2 MinPixelPos, int2 MaxPixelPos)
@@ -434,6 +442,7 @@ taa_short2 ClampPixelOffset(
 	return SamplePixelPos;
 }
 
+CALL_SITE_DEBUGLOC
 taa_short2x2 AddAndClampPixelOffset(
 	taa_short2x2 KernelCenterPixelPos,
 	taa_short2x2 Offset,
@@ -470,6 +479,7 @@ taa_short2x2 AddAndClampPixelOffset(
 	return dpv_interleave_registers(SamplePixelPos0, SamplePixelPos1);
 }
 
+CALL_SITE_DEBUGLOC
 taa_short2 AddAndClampPixelOffset(
 	taa_short2 KernelCenterPixelPos,
 	taa_short2 Offset,
@@ -493,6 +503,7 @@ taa_short2 AddAndClampPixelOffset(
 	return SamplePixelPos;
 }
 
+CALL_SITE_DEBUGLOC
 taa_short2x2 InvalidateOutputPixelPos(taa_short2x2 PixelPos, uint2 ViewportMax)
 #if 1
 {
@@ -515,6 +526,7 @@ taa_short2x2 InvalidateOutputPixelPos(taa_short2x2 PixelPos, uint2 ViewportMax)
 }
 #endif
 
+CALL_SITE_DEBUGLOC
 taa_short2 InvalidateOutputPixelPos(taa_short2 PixelPos, uint2 ViewportMax)
 #if 1
 {
@@ -537,6 +549,7 @@ taa_short2 InvalidateOutputPixelPos(taa_short2 PixelPos, uint2 ViewportMax)
 }
 #endif
 
+CALL_SITE_DEBUGLOC
 taa_ushort2 Map8x8Tile2x2Lane(uint GroupThreadIndex)
 {
 	taa_ushort2 T = taa_ushort(GroupThreadIndex).xx;
@@ -544,6 +557,7 @@ taa_ushort2 Map8x8Tile2x2Lane(uint GroupThreadIndex)
 	return GroupId;
 }
 
+CALL_SITE_DEBUGLOC
 taa_ushort2 Map16x16Tile2x2Lane(uint GroupThreadIndex)
 {
 	taa_ushort T = taa_ushort(GroupThreadIndex);
@@ -555,11 +569,13 @@ taa_ushort2 Map16x16Tile2x2Lane(uint GroupThreadIndex)
 	return GroupId;
 }
 
+CALL_SITE_DEBUGLOC
 float SafeRcp(float x)
 {
 	return x > 0.0 ? rcp(x) : 0.0;
 }
 
+CALL_SITE_DEBUGLOC
 float2 SafeRcp(float2 x)
 {
 	return float2(SafeRcp(x.x), SafeRcp(x.y));
@@ -567,12 +583,14 @@ float2 SafeRcp(float2 x)
 
 #if CONFIG_COMPILE_FP16
 
+CALL_SITE_DEBUGLOC
 half SafeRcp(half x)
 {
 	// If x=0.0, then MaxHalfFloat * 0.0 = 0.0
 	return min(rcp(x), half(MaxHalfFloat)) * saturate(x * asfloat16(uint16_t(0x7C00)));
 }
 
+CALL_SITE_DEBUGLOC
 half2 SafeRcp(half2 x)
 {
 	// If x=0.0, then MaxHalfFloat * 0.0 = 0.0
@@ -581,16 +599,19 @@ half2 SafeRcp(half2 x)
 
 #endif
 
+CALL_SITE_DEBUGLOC
 taa_half2 ComputeLowFrequencyContribution(taa_half2 HistoryValidity)
 {
 	return saturate((HistoryValidity - taa_half(1.0 * rcp(MAX_SAMPLE_COUNT))) * taa_half(6.0));
 }
 
+CALL_SITE_DEBUGLOC
 taa_half ComputeLowFrequencyContribution(taa_half HistoryValidity)
 {
 	return dpv_lo(ComputeLowFrequencyContribution(dpv_interleave_mono_registers(HistoryValidity)));
 }
 
+CALL_SITE_DEBUGLOC
 taa_half3x2 RGBToYCoCg(taa_half3x2 RGB)
 {
     taa_half2 R = RGB[0];
@@ -606,6 +627,7 @@ taa_half3x2 RGBToYCoCg(taa_half3x2 RGB)
     return YCoCg;
 }
 
+CALL_SITE_DEBUGLOC
 taa_half3x2 YCoCgToRGB(taa_half3x2 YCoCg)
 {
 	taa_half2 Y =  YCoCg[0] * taa_half(0.25);
@@ -620,11 +642,13 @@ taa_half3x2 YCoCgToRGB(taa_half3x2 YCoCg)
 	return RGB;
 }
 
+CALL_SITE_DEBUGLOC
 taa_half3 RGBToYCoCg(taa_half3 RGB)
 {
 	return dpv_lo(RGBToYCoCg(dpv_interleave_mono_registers(RGB)));
 }
 
+CALL_SITE_DEBUGLOC
 taa_half3 YCoCgToRGB(taa_half3 YCoCg)
 {
 	return dpv_lo(YCoCgToRGB(dpv_interleave_mono_registers(YCoCg)));
@@ -636,16 +660,19 @@ taa_half3 YCoCgToRGB(taa_half3 YCoCg)
 
 // Faster but less accurate luma computation. 
 // Luma includes a scaling by 4.
+CALL_SITE_DEBUGLOC
 taa_half2 Luma4(taa_half3x2 Color)
 {
 	return (Color[1] * taa_half(2.0)) + (Color[0] + Color[2]);
 }
 
+CALL_SITE_DEBUGLOC
 taa_half Luma4(taa_half3 Color)
 {
 	return dpv_lo(Luma4(dpv_interleave_mono_registers(Color)));
 }
 
+CALL_SITE_DEBUGLOC
 taa_half2 HdrWeightY(taa_half2 Luma)
 {
 	taa_half Exposure = taa_half(1.0);
@@ -653,27 +680,32 @@ taa_half2 HdrWeightY(taa_half2 Luma)
 	return max(taa_half(HDR_WEIGHT_SAFE_MIN_VALUE), rcp(Luma * Exposure + taa_half(4.0)));
 }
 
+CALL_SITE_DEBUGLOC
 taa_half HdrWeightY(taa_half Luma)
 {
 	return dpv_lo(HdrWeightY(dpv_interleave_mono_registers(Luma)));
 }
 
+CALL_SITE_DEBUGLOC
 taa_half2 HdrWeightInvY(taa_half2 LDRLuma) 
 {
 	return taa_half(4.0) * rcp(taa_half(1.0) - LDRLuma);
 }
 
+CALL_SITE_DEBUGLOC
 taa_half HdrWeightInvY(taa_half LDRLuma) 
 {
 	return dpv_lo(HdrWeightInvY(dpv_interleave_mono_registers(LDRLuma)));
 }
 
 // Optimized HDR weighting function.
+CALL_SITE_DEBUGLOC
 taa_half2 HdrWeight4(taa_half3x2 Color)
 {
 	return HdrWeightY(Luma4(Color));
 }
 
+CALL_SITE_DEBUGLOC
 taa_half HdrWeight4(taa_half3 Color)
 {
 	return HdrWeightY(Luma4(Color));
@@ -681,6 +713,7 @@ taa_half HdrWeight4(taa_half3 Color)
 
 
 // Returns the weight of a pixels at a coordinate <PixelDelta> from the PDF highest point.
+CALL_SITE_DEBUGLOC
 taa_half2 ComputeSampleWeigth(taa_half2 UpscaleFactor, taa_half2x2 PixelDelta, const float MinimalContribution)
 {
 	taa_half2 u2 = UpscaleFactor * UpscaleFactor;
@@ -691,6 +724,7 @@ taa_half2 ComputeSampleWeigth(taa_half2 UpscaleFactor, taa_half2x2 PixelDelta, c
 	return saturate((taa_half(0.9) * x2 - taa_half(1.9)) * x2 + taa_half(1.0 + MinimalContribution));
 }
 
+CALL_SITE_DEBUGLOC
 taa_half ComputeSampleWeigth(taa_half UpscaleFactor, taa_half2 PixelDelta, const float MinimalContribution)
 {
 	return dpv_lo(ComputeSampleWeigth(
@@ -699,6 +733,7 @@ taa_half ComputeSampleWeigth(taa_half UpscaleFactor, taa_half2 PixelDelta, const
 		MinimalContribution));
 }
 
+CALL_SITE_DEBUGLOC
 taa_half2x2 WeightedLerpFactors(taa_half2 WeightA, taa_half2 WeightB, taa_half2 Blend)
 {
 	taa_half2 BlendA = (taa_half(1.0) - Blend) * WeightA;
@@ -709,6 +744,7 @@ taa_half2x2 WeightedLerpFactors(taa_half2 WeightA, taa_half2 WeightB, taa_half2
 	return taa_half2x2(BlendA, BlendB);
 }
 
+CALL_SITE_DEBUGLOC
 taa_half2 WeightedLerpFactors(taa_half WeightA, taa_half WeightB, taa_half Blend)
 {
 	return dpv_lo(WeightedLerpFactors(
@@ -717,6 +753,7 @@ taa_half2 WeightedLerpFactors(taa_half WeightA, taa_half WeightB, taa_half Blend
 		dpv_interleave_mono_registers(Blend)));
 }
 
+CALL_SITE_DEBUGLOC
 bool TakeOnlyOneSamplePair(float2 Offset)
 {
 	return Offset.y > 0.0 || (Offset.x > 0.0 && Offset.y == 0.0);
@@ -737,6 +774,7 @@ taa_half ComputePredictionCompleteness(taa_half SampleHistoryValidity)
 	return saturate(SampleHistoryValidity * taa_half(MAX_SAMPLE_COUNT) - taa_half(0.2));
 }
 
+CALL_SITE_DEBUGLOC
 bool IsOffScreen(uint bCameraCut, float2 ScreenPos, taa_half ParallaxRejectionMask)
 {
 	bool bIsCameraCut = bCameraCut != 0;
@@ -883,6 +921,7 @@ taa_subpixel_payload GetClosestSubpixelPayload(FSubpixelNeighborhood Neighborhoo
 	return SubpixelPayload;
 }
 
+CALL_SITE_DEBUGLOC
 taa_subpixel_payload CompressParallaxFactor(taa_half ParallaxFactor)
 {
 	return clamp(
@@ -890,6 +929,7 @@ taa_subpixel_payload CompressParallaxFactor(taa_half ParallaxFactor)
 		taa_subpixel_payload(0), taa_subpixel_payload(SUB_PIXEL_PARALLAX_FACTOR_BIT_MASK));
 }
 
+CALL_SITE_DEBUGLOC
 taa_half UncompressParallaxFactor(taa_subpixel_payload CompressedParallaxFactor)
 {
 	return 
@@ -898,6 +938,7 @@ taa_half UncompressParallaxFactor(taa_subpixel_payload CompressedParallaxFactor)
 		taa_half(1.0);
 }
 
+CALL_SITE_DEBUGLOC
 taa_half3x2 QuantizeForFloatRenderTarget(taa_half3x2 Color, taa_half E, const float3 QuantizationError)
 {
 	taa_half3x2 Error = dpv_mul(Color, taa_half3(QuantizationError));
@@ -921,6 +962,7 @@ taa_half3x2 QuantizeForFloatRenderTarget(taa_half3x2 Color, taa_half E, const fl
 	return Color + Error * E;
 }
 
+CALL_SITE_DEBUGLOC
 taa_half3 QuantizeForFloatRenderTarget(taa_half3 Color, taa_half E, const float3 QuantizationError)
 {
 	return dpv_lo(QuantizeForFloatRenderTarget(dpv_interleave_mono_registers(Color), E, QuantizationError));
@@ -941,6 +983,7 @@ void ComputeCompressedUseCountPixelCoordinates(uint2 PixelPos, out uint2 Compres
 	
 #endif
 
+CALL_SITE_DEBUGLOC
 uint2 EncodeSpatialAntiAliasingOffset(taa_half2x2 TexelOffset)
 {
 	const uint Mask = (1 << SPATIAL_ANTI_ALIASING_OFFSET_BITDEPTH_PER_CHANNEL) - 1;
@@ -954,11 +997,13 @@ uint2 EncodeSpatialAntiAliasingOffset(taa_half2x2 TexelOffset)
 	return EncodedTexelOffset;
 }
 
+CALL_SITE_DEBUGLOC
 uint EncodeSpatialAntiAliasingOffset(taa_half2 TexelOffset)
 {
 	return dpv_lo(EncodeSpatialAntiAliasingOffset(dpv_interleave_mono_registers(TexelOffset)));
 }
 
+CALL_SITE_DEBUGLOC
 taa_half2x2 DecodeSpatialAntiAliasingOffset(taa_ushort2 EncodedInputTexelOffset)
 {
 	const taa_ushort Mask = taa_ushort(1 << SPATIAL_ANTI_ALIASING_OFFSET_BITDEPTH_PER_CHANNEL) - taa_ushort(1);
@@ -973,6 +1018,7 @@ taa_half2x2 DecodeSpatialAntiAliasingOffset(taa_ushort2 EncodedInputTexelOffset)
 	return TexelOffset;
 }
 
+CALL_SITE_DEBUGLOC
 taa_half2 DecodeSpatialAntiAliasingOffset(taa_ushort EncodedInputTexelOffset)
 {
 	return dpv_lo(DecodeSpatialAntiAliasingOffset(dpv_interleave_mono_registers(EncodedInputTexelOffset)));
@@ -984,6 +1030,7 @@ taa_half2 DecodeSpatialAntiAliasingOffset(taa_ushort EncodedInputTexelOffset)
 #define VELOCITY_HOLE_FILLING_ANGLE_BITS (VELOCITY_HOLE_FILLING_BITS - VELOCITY_HOLE_FILLING_LENGTH_BITS)
 #define VELOCITY_HOLE_FILLING_LENGTH_PRECISION (pow(2.0, 5))
 
+CALL_SITE_DEBUGLOC
 uint EncodeHoleFillingVelocity(float2 PixelVelocity)
 {
 	const uint AngleBitDepth = 1 << VELOCITY_HOLE_FILLING_ANGLE_BITS;
@@ -1000,6 +1047,7 @@ uint EncodeHoleFillingVelocity(float2 PixelVelocity)
 	return EncodedHoleFillingVelocity;
 }
 
+CALL_SITE_DEBUGLOC
 void DecodeHoleFillingVelocity(uint EncodedHoleFillingVelocity, out float2 PixelVelocity, out float OutOfRangeFactor)
 {
 	const uint AngleBitDepth = 1 << VELOCITY_HOLE_FILLING_ANGLE_BITS;
@@ -1022,6 +1070,7 @@ void DecodeHoleFillingVelocity(uint EncodedHoleFillingVelocity, out float2 Pixel
 	OutOfRangeFactor = saturate(LengthRange - Length);
 }
 
+CALL_SITE_DEBUGLOC
 taa_short2 GetLaneOffsetSign()
 #if CONFIG_BUTTERFLY_KERNEL
 {
@@ -1035,11 +1084,13 @@ taa_short2 GetLaneOffsetSign()
 }
 #endif
 
+CALL_SITE_DEBUGLOC
 float ComputePixelVelocityBilateralWeight(float2 PixelVelocityDelta)
 {
 	return saturate(1.0 - length2(PixelVelocityDelta));
 }
 
+CALL_SITE_DEBUGLOC
 float ComputeScreenVelocityBilateralWeight(float2 ScreenVelocityDelta)
 {
 	return ComputePixelVelocityBilateralWeight(ScreenVelocityDelta * ScreenVelocityToInputPixelVelocity);
diff --git a/Engine/Shaders/Private/TemporalSuperResolution/TSRUpdateHistory.usf b/Engine/Shaders/Private/TemporalSuperResolution/TSRUpdateHistory.usf
--- a/Engine/Shaders/Private/TemporalSuperResolution/TSRUpdateHistory.usf
+++ b/Engine/Shaders/Private/TemporalSuperResolution/TSRUpdateHistory.usf
@@ -125,6 +125,7 @@ groupshared taa_subpixel_payload2 SharedParallaxFactorBits[GROUP_SHARED_SIZE];
 
 //------------------------------------------------------- FUNCTIONS
 
+CALL_SITE_DEBUGLOC 
 taa_ushort2x2 Map8x8Tile2x2LaneDPV(uint GroupThreadIndex)
 {
 	taa_ushort2 GroupId = Map8x8Tile2x2Lane(GroupThreadIndex);
@@ -135,11 +136,13 @@ taa_ushort2x2 Map8x8Tile2x2LaneDPV(uint GroupThreadIndex)
 #endif
 }
 
+CALL_SITE_DEBUGLOC 
 taa_half3x2 TransformColorForClampingBox(taa_half3x2 Color)
 {
 	return RGBToYCoCg(Color);
 }
 
+CALL_SITE_DEBUGLOC 
 taa_half3x2 GetColorChannel(taa_half4x2 Color)
 {
 	return dpv_interleave_registers(
diff --git a/Engine/Shaders/Public/DualPixelVectorization.ush b/Engine/Shaders/Public/DualPixelVectorization.ush
--- a/Engine/Shaders/Public/DualPixelVectorization.ush
+++ b/Engine/Shaders/Public/DualPixelVectorization.ush
@@ -201,16 +201,16 @@
  */
 
 #define DEFINE_DPV_FUNCTION_SCALAR(TYPE) \
-	__DPV_MULTI_SCALAR(TYPE) dpv_interleave_registers(TYPE lo, TYPE hi) { return __DPV_MULTI_SCALAR(TYPE)(lo, hi); } \
-	TYPE dpv_lo(__DPV_MULTI_SCALAR(TYPE) x) { return x[0]; } \
-	TYPE dpv_hi(__DPV_MULTI_SCALAR(TYPE) x) { return x[1]; } \
-	TYPE dpv_access_pixel(__DPV_MULTI_SCALAR(TYPE) x, const uint DualPixelId) { return x[DualPixelId]; } \
+	CALL_SITE_DEBUGLOC __DPV_MULTI_SCALAR(TYPE) dpv_interleave_registers(TYPE lo, TYPE hi) { return __DPV_MULTI_SCALAR(TYPE)(lo, hi); } \
+	CALL_SITE_DEBUGLOC TYPE dpv_lo(__DPV_MULTI_SCALAR(TYPE) x) { return x[0]; } \
+	CALL_SITE_DEBUGLOC TYPE dpv_hi(__DPV_MULTI_SCALAR(TYPE) x) { return x[1]; } \
+	CALL_SITE_DEBUGLOC TYPE dpv_access_pixel(__DPV_MULTI_SCALAR(TYPE) x, const uint DualPixelId) { return x[DualPixelId]; } \
 
 #define DEFINE_DPV_FUNCTION_VECTOR(TYPE, COMPONENTS) \
-	__DPV_MULTI_VECTOR(TYPE, COMPONENTS) dpv_interleave_registers(__DPV_VECTOR(TYPE, COMPONENTS) lo, __DPV_VECTOR(TYPE, COMPONENTS) hi) { return transpose(__DPV_MULTI_VECTOR_TRANSPOSED(TYPE, COMPONENTS)(lo, hi)); } \
-	__DPV_VECTOR(TYPE, COMPONENTS) dpv_lo(__DPV_MULTI_VECTOR(TYPE, COMPONENTS) x) { return transpose(x)[0]; } \
-	__DPV_VECTOR(TYPE, COMPONENTS) dpv_hi(__DPV_MULTI_VECTOR(TYPE, COMPONENTS) x) { return transpose(x)[1]; } \
-	__DPV_VECTOR(TYPE, COMPONENTS) dpv_access_pixel(__DPV_MULTI_VECTOR(TYPE, COMPONENTS) x, const uint DualPixelId) { return transpose(x)[DualPixelId]; } \
+	CALL_SITE_DEBUGLOC __DPV_MULTI_VECTOR(TYPE, COMPONENTS) dpv_interleave_registers(__DPV_VECTOR(TYPE, COMPONENTS) lo, __DPV_VECTOR(TYPE, COMPONENTS) hi) { return transpose(__DPV_MULTI_VECTOR_TRANSPOSED(TYPE, COMPONENTS)(lo, hi)); } \
+	CALL_SITE_DEBUGLOC __DPV_VECTOR(TYPE, COMPONENTS) dpv_lo(__DPV_MULTI_VECTOR(TYPE, COMPONENTS) x) { return transpose(x)[0]; } \
+	CALL_SITE_DEBUGLOC __DPV_VECTOR(TYPE, COMPONENTS) dpv_hi(__DPV_MULTI_VECTOR(TYPE, COMPONENTS) x) { return transpose(x)[1]; } \
+	CALL_SITE_DEBUGLOC __DPV_VECTOR(TYPE, COMPONENTS) dpv_access_pixel(__DPV_MULTI_VECTOR(TYPE, COMPONENTS) x, const uint DualPixelId) { return transpose(x)[DualPixelId]; } \
 
 DEFINE_DPV_FUNCTION_FOR_ALL_SCALARS(DEFINE_DPV_FUNCTION_SCALAR)
 DEFINE_DPV_FUNCTION_FOR_ALL_VECTORS(DEFINE_DPV_FUNCTION_VECTOR)
@@ -227,10 +227,10 @@ DEFINE_DPV_FUNCTION_FOR_ALL_VECTORS(DEFINE_DPV_FUNCTION_VECTOR)
  */
 
 #define DEFINE_DPV_FUNCTION_SCALAR(TYPE) \
-	__DPV_MULTI_SCALAR(TYPE) dpv_force_interleave_registers(TYPE lo, TYPE hi) { return dpv_interleave_registers(lo, hi); } \
+	CALL_SITE_DEBUGLOC __DPV_MULTI_SCALAR(TYPE) dpv_force_interleave_registers(TYPE lo, TYPE hi) { return dpv_interleave_registers(lo, hi); } \
 
 #define DEFINE_DPV_FUNCTION_VECTOR(TYPE, COMPONENTS) \
-	__DPV_MULTI_VECTOR(TYPE, COMPONENTS) dpv_force_interleave_registers(__DPV_VECTOR(TYPE, COMPONENTS) lo, __DPV_VECTOR(TYPE, COMPONENTS) hi) { return dpv_interleave_registers(lo, hi); } \
+	CALL_SITE_DEBUGLOC __DPV_MULTI_VECTOR(TYPE, COMPONENTS) dpv_force_interleave_registers(__DPV_VECTOR(TYPE, COMPONENTS) lo, __DPV_VECTOR(TYPE, COMPONENTS) hi) { return dpv_interleave_registers(lo, hi); } \
 
 DEFINE_DPV_FUNCTION_FOR_ALL_32BIT_SCALARS(DEFINE_DPV_FUNCTION_SCALAR)
 DEFINE_DPV_FUNCTION_FOR_ALL_32BIT_VECTORS(DEFINE_DPV_FUNCTION_VECTOR)
@@ -239,10 +239,10 @@ DEFINE_DPV_FUNCTION_FOR_ALL_32BIT_VECTORS(DEFINE_DPV_FUNCTION_VECTOR)
 
 
 #define DEFINE_DPV_FUNCTION_SCALAR(TYPE) \
-	__DPV_MULTI_SCALAR(TYPE) dpv_force_interleave_registers(TYPE lo, TYPE hi) { return v_pack_b32_b16(lo, hi); } \
+	CALL_SITE_DEBUGLOC __DPV_MULTI_SCALAR(TYPE) dpv_force_interleave_registers(TYPE lo, TYPE hi) { return v_pack_b32_b16(lo, hi); } \
 
 #define DEFINE_DPV_FUNCTION_VECTOR(TYPE, COMPONENTS) \
-	__DPV_MULTI_VECTOR(TYPE, COMPONENTS) dpv_force_interleave_registers(__DPV_VECTOR(TYPE, COMPONENTS) lo, __DPV_VECTOR(TYPE, COMPONENTS) hi) { \
+	CALL_SITE_DEBUGLOC __DPV_MULTI_VECTOR(TYPE, COMPONENTS) dpv_force_interleave_registers(__DPV_VECTOR(TYPE, COMPONENTS) lo, __DPV_VECTOR(TYPE, COMPONENTS) hi) { \
 		return __DPV_FOREACH_COMPONENTS_CONSTRUCTOR(TYPE, COMPONENTS, __DPV_CODE_FUNCTION_2_ARGS, dpv_force_interleave_registers, lo, hi); \
 	} \
 
@@ -282,10 +282,10 @@ DEFINE_DPV_FUNCTION_FOR_ALL_16BIT_VECTORS(DEFINE_DPV_FUNCTION_VECTOR)
  *		}
  */
 #define DEFINE_DPV_FUNCTION_SCALAR(TYPE) \
-	__DPV_MULTI_SCALAR(TYPE) dpv_interleave_mono_registers(TYPE lo) { return dpv_interleave_registers(lo, lo); }
+	CALL_SITE_DEBUGLOC __DPV_MULTI_SCALAR(TYPE) dpv_interleave_mono_registers(TYPE lo) { return dpv_interleave_registers(lo, lo); }
 
 #define DEFINE_DPV_FUNCTION_VECTOR(TYPE, COMPONENTS) \
-	__DPV_MULTI_VECTOR(TYPE, COMPONENTS) dpv_interleave_mono_registers(__DPV_VECTOR(TYPE, COMPONENTS) lo) { return dpv_interleave_registers(lo, lo); }
+	CALL_SITE_DEBUGLOC __DPV_MULTI_VECTOR(TYPE, COMPONENTS) dpv_interleave_mono_registers(__DPV_VECTOR(TYPE, COMPONENTS) lo) { return dpv_interleave_registers(lo, lo); }
 
 DEFINE_DPV_FUNCTION_FOR_ALL_SCALARS(DEFINE_DPV_FUNCTION_SCALAR)
 DEFINE_DPV_FUNCTION_FOR_ALL_VECTORS(DEFINE_DPV_FUNCTION_VECTOR)
@@ -303,25 +303,25 @@ DEFINE_DPV_FUNCTION_FOR_ALL_VECTORS(DEFINE_DPV_FUNCTION_VECTOR)
  *		half3x2 PremultipliedColor = dpv_scale(Color, Coverage);
  */
 #define DEFINE_DPV_FUNCTION_VECTOR(TYPE, COMPONENTS) \
-	__DPV_MULTI_VECTOR(TYPE, COMPONENTS) dpv_mul(__DPV_MULTI_VECTOR(TYPE, COMPONENTS) a, __DPV_MULTI_VECTOR(TYPE, COMPONENTS) b) { \
+	CALL_SITE_DEBUGLOC __DPV_MULTI_VECTOR(TYPE, COMPONENTS) dpv_mul(__DPV_MULTI_VECTOR(TYPE, COMPONENTS) a, __DPV_MULTI_VECTOR(TYPE, COMPONENTS) b) { \
 		return __DPV_FOREACH_COMPONENTS_CONSTRUCTOR(TYPE, COMPONENTS, __DPV_CODE_OPERATOR_VECTOR_VECTOR, a, *, b); \
 	} \
-	__DPV_MULTI_VECTOR(TYPE, COMPONENTS) dpv_mul(__DPV_MULTI_VECTOR(TYPE, COMPONENTS) a, __DPV_VECTOR(TYPE, COMPONENTS) b) { \
+	CALL_SITE_DEBUGLOC __DPV_MULTI_VECTOR(TYPE, COMPONENTS) dpv_mul(__DPV_MULTI_VECTOR(TYPE, COMPONENTS) a, __DPV_VECTOR(TYPE, COMPONENTS) b) { \
 		return __DPV_FOREACH_COMPONENTS_CONSTRUCTOR(TYPE, COMPONENTS, __DPV_CODE_OPERATOR_VECTOR_VECTOR, a, *, b); \
 	} \
-	__DPV_MULTI_VECTOR(TYPE, COMPONENTS) dpv_mul(__DPV_VECTOR(TYPE, COMPONENTS) a, __DPV_MULTI_VECTOR(TYPE, COMPONENTS) b) { \
+	CALL_SITE_DEBUGLOC __DPV_MULTI_VECTOR(TYPE, COMPONENTS) dpv_mul(__DPV_VECTOR(TYPE, COMPONENTS) a, __DPV_MULTI_VECTOR(TYPE, COMPONENTS) b) { \
 		return __DPV_FOREACH_COMPONENTS_CONSTRUCTOR(TYPE, COMPONENTS, __DPV_CODE_OPERATOR_VECTOR_VECTOR, a, *, b); \
 	} \
-	__DPV_MULTI_VECTOR(TYPE, COMPONENTS) dpv_scale(__DPV_MULTI_VECTOR(TYPE, COMPONENTS) a, __DPV_MULTI_SCALAR(TYPE) b) { \
+	CALL_SITE_DEBUGLOC __DPV_MULTI_VECTOR(TYPE, COMPONENTS) dpv_scale(__DPV_MULTI_VECTOR(TYPE, COMPONENTS) a, __DPV_MULTI_SCALAR(TYPE) b) { \
 		return __DPV_FOREACH_COMPONENTS_CONSTRUCTOR(TYPE, COMPONENTS, __DPV_CODE_OPERATOR_VECTOR_SCALAR, a, *, b); \
 	} \
-	__DPV_MULTI_VECTOR(TYPE, COMPONENTS) dpv_scale(__DPV_MULTI_SCALAR(TYPE) a, __DPV_MULTI_VECTOR(TYPE, COMPONENTS) b) { \
+	CALL_SITE_DEBUGLOC __DPV_MULTI_VECTOR(TYPE, COMPONENTS) dpv_scale(__DPV_MULTI_SCALAR(TYPE) a, __DPV_MULTI_VECTOR(TYPE, COMPONENTS) b) { \
 		return __DPV_FOREACH_COMPONENTS_CONSTRUCTOR(TYPE, COMPONENTS, __DPV_CODE_OPERATOR_VECTOR_SCALAR, b, *, a); \
 	} \
-	__DPV_MULTI_VECTOR(TYPE, COMPONENTS) dpv_scale(__DPV_MULTI_VECTOR(TYPE, COMPONENTS) a, TYPE b) { \
+	CALL_SITE_DEBUGLOC __DPV_MULTI_VECTOR(TYPE, COMPONENTS) dpv_scale(__DPV_MULTI_VECTOR(TYPE, COMPONENTS) a, TYPE b) { \
 		return __DPV_FOREACH_COMPONENTS_CONSTRUCTOR(TYPE, COMPONENTS, __DPV_CODE_OPERATOR_VECTOR_SCALAR, a, *, b); \
 	} \
-	__DPV_MULTI_VECTOR(TYPE, COMPONENTS) dpv_scale(TYPE a, __DPV_MULTI_VECTOR(TYPE, COMPONENTS) b) { \
+	CALL_SITE_DEBUGLOC __DPV_MULTI_VECTOR(TYPE, COMPONENTS) dpv_scale(TYPE a, __DPV_MULTI_VECTOR(TYPE, COMPONENTS) b) { \
 		return __DPV_FOREACH_COMPONENTS_CONSTRUCTOR(TYPE, COMPONENTS, __DPV_CODE_OPERATOR_VECTOR_SCALAR, b, *, a); \
 	} \
 
@@ -338,22 +338,22 @@ DEFINE_DPV_FUNCTION_FOR_ALL_VECTORS(DEFINE_DPV_FUNCTION_VECTOR)
  *		half3x2 PosB = dpv_add(PosA, PosAToB);
  */
 #define DEFINE_DPV_FUNCTION_VECTOR(TYPE, COMPONENTS) \
-	__DPV_MULTI_VECTOR(TYPE, COMPONENTS) dpv_add(__DPV_MULTI_VECTOR(TYPE, COMPONENTS) a, __DPV_MULTI_VECTOR(TYPE, COMPONENTS) b) { \
+	CALL_SITE_DEBUGLOC __DPV_MULTI_VECTOR(TYPE, COMPONENTS) dpv_add(__DPV_MULTI_VECTOR(TYPE, COMPONENTS) a, __DPV_MULTI_VECTOR(TYPE, COMPONENTS) b) { \
 		return __DPV_FOREACH_COMPONENTS_CONSTRUCTOR(TYPE, COMPONENTS, __DPV_CODE_OPERATOR_VECTOR_VECTOR, a, +, b); \
 	} \
-	__DPV_MULTI_VECTOR(TYPE, COMPONENTS) dpv_add(__DPV_MULTI_VECTOR(TYPE, COMPONENTS) a, __DPV_VECTOR(TYPE, COMPONENTS) b) { \
+	CALL_SITE_DEBUGLOC __DPV_MULTI_VECTOR(TYPE, COMPONENTS) dpv_add(__DPV_MULTI_VECTOR(TYPE, COMPONENTS) a, __DPV_VECTOR(TYPE, COMPONENTS) b) { \
 		return __DPV_FOREACH_COMPONENTS_CONSTRUCTOR(TYPE, COMPONENTS, __DPV_CODE_OPERATOR_VECTOR_VECTOR, a, +, b); \
 	} \
-	__DPV_MULTI_VECTOR(TYPE, COMPONENTS) dpv_add(__DPV_VECTOR(TYPE, COMPONENTS) a, __DPV_MULTI_VECTOR(TYPE, COMPONENTS) b) { \
+	CALL_SITE_DEBUGLOC __DPV_MULTI_VECTOR(TYPE, COMPONENTS) dpv_add(__DPV_VECTOR(TYPE, COMPONENTS) a, __DPV_MULTI_VECTOR(TYPE, COMPONENTS) b) { \
 		return __DPV_FOREACH_COMPONENTS_CONSTRUCTOR(TYPE, COMPONENTS, __DPV_CODE_OPERATOR_VECTOR_VECTOR, a, +, b); \
 	} \
-	__DPV_MULTI_VECTOR(TYPE, COMPONENTS) dpv_sub(__DPV_MULTI_VECTOR(TYPE, COMPONENTS) a, __DPV_MULTI_VECTOR(TYPE, COMPONENTS) b) { \
+	CALL_SITE_DEBUGLOC __DPV_MULTI_VECTOR(TYPE, COMPONENTS) dpv_sub(__DPV_MULTI_VECTOR(TYPE, COMPONENTS) a, __DPV_MULTI_VECTOR(TYPE, COMPONENTS) b) { \
 		return __DPV_FOREACH_COMPONENTS_CONSTRUCTOR(TYPE, COMPONENTS, __DPV_CODE_OPERATOR_VECTOR_VECTOR, a, -, b); \
 	} \
-	__DPV_MULTI_VECTOR(TYPE, COMPONENTS) dpv_sub(__DPV_MULTI_VECTOR(TYPE, COMPONENTS) a, __DPV_VECTOR(TYPE, COMPONENTS) b) { \
+	CALL_SITE_DEBUGLOC __DPV_MULTI_VECTOR(TYPE, COMPONENTS) dpv_sub(__DPV_MULTI_VECTOR(TYPE, COMPONENTS) a, __DPV_VECTOR(TYPE, COMPONENTS) b) { \
 		return __DPV_FOREACH_COMPONENTS_CONSTRUCTOR(TYPE, COMPONENTS, __DPV_CODE_OPERATOR_VECTOR_VECTOR, a, -, b); \
 	} \
-	__DPV_MULTI_VECTOR(TYPE, COMPONENTS) dpv_sub(__DPV_VECTOR(TYPE, COMPONENTS) a, __DPV_MULTI_VECTOR(TYPE, COMPONENTS) b) { \
+	CALL_SITE_DEBUGLOC __DPV_MULTI_VECTOR(TYPE, COMPONENTS) dpv_sub(__DPV_VECTOR(TYPE, COMPONENTS) a, __DPV_MULTI_VECTOR(TYPE, COMPONENTS) b) { \
 		return __DPV_FOREACH_COMPONENTS_CONSTRUCTOR(TYPE, COMPONENTS, __DPV_CODE_OPERATOR_VECTOR_VECTOR, a, -, b); \
 	} \
 
@@ -369,34 +369,34 @@ DEFINE_DPV_FUNCTION_FOR_ALL_VECTORS(DEFINE_DPV_FUNCTION_VECTOR)
  *		half2 Luma = dpv_dot(Color, half3(0.25, 0.5, 0.25));
  */
 #define DEFINE_DPV_FUNCTION_VECTOR(TYPE, COMPONENTS) \
-	__DPV_MULTI_SCALAR(TYPE) dpv_dot(__DPV_MULTI_VECTOR(TYPE, COMPONENTS) a, __DPV_MULTI_VECTOR(TYPE, COMPONENTS) b) { \
+	CALL_SITE_DEBUGLOC __DPV_MULTI_SCALAR(TYPE) dpv_dot(__DPV_MULTI_VECTOR(TYPE, COMPONENTS) a, __DPV_MULTI_VECTOR(TYPE, COMPONENTS) b) { \
 		return __DPV_FOREACH_COMPONENTS(COMPONENTS, __DPV_CODE_OPERATOR_VECTOR_VECTOR, +, a, *, b); \
 	} \
-	__DPV_MULTI_SCALAR(TYPE) dpv_dot(__DPV_VECTOR(TYPE, COMPONENTS) a, __DPV_MULTI_VECTOR(TYPE, COMPONENTS) b) { \
+	CALL_SITE_DEBUGLOC __DPV_MULTI_SCALAR(TYPE) dpv_dot(__DPV_VECTOR(TYPE, COMPONENTS) a, __DPV_MULTI_VECTOR(TYPE, COMPONENTS) b) { \
 		return __DPV_FOREACH_COMPONENTS(COMPONENTS, __DPV_CODE_OPERATOR_VECTOR_VECTOR, +, a, *, b); \
 	} \
-	__DPV_MULTI_SCALAR(TYPE) dpv_dot(__DPV_MULTI_VECTOR(TYPE, COMPONENTS) a, __DPV_VECTOR(TYPE, COMPONENTS) b) { \
+	CALL_SITE_DEBUGLOC __DPV_MULTI_SCALAR(TYPE) dpv_dot(__DPV_MULTI_VECTOR(TYPE, COMPONENTS) a, __DPV_VECTOR(TYPE, COMPONENTS) b) { \
 		return __DPV_FOREACH_COMPONENTS(COMPONENTS, __DPV_CODE_OPERATOR_VECTOR_VECTOR, +, a, *, b); \
 	} \
-	__DPV_MULTI_VECTOR(TYPE, COMPONENTS) dpv_clamp(__DPV_MULTI_VECTOR(TYPE, COMPONENTS) x, __DPV_MULTI_VECTOR(TYPE, COMPONENTS) min, __DPV_MULTI_VECTOR(TYPE, COMPONENTS) max) { \
+	CALL_SITE_DEBUGLOC __DPV_MULTI_VECTOR(TYPE, COMPONENTS) dpv_clamp(__DPV_MULTI_VECTOR(TYPE, COMPONENTS) x, __DPV_MULTI_VECTOR(TYPE, COMPONENTS) min, __DPV_MULTI_VECTOR(TYPE, COMPONENTS) max) { \
 		return __DPV_FOREACH_COMPONENTS_CONSTRUCTOR(TYPE, COMPONENTS, __DPV_CODE_FUNCTION_3_ARGS, clamp, x, min, max); \
 	} \
-	__DPV_MULTI_VECTOR(TYPE, COMPONENTS) dpv_min(__DPV_MULTI_VECTOR(TYPE, COMPONENTS) a, __DPV_MULTI_VECTOR(TYPE, COMPONENTS) b) { \
+	CALL_SITE_DEBUGLOC __DPV_MULTI_VECTOR(TYPE, COMPONENTS) dpv_min(__DPV_MULTI_VECTOR(TYPE, COMPONENTS) a, __DPV_MULTI_VECTOR(TYPE, COMPONENTS) b) { \
 		return min(a, b); \
 	} \
-	__DPV_MULTI_VECTOR(TYPE, COMPONENTS) dpv_min(__DPV_VECTOR(TYPE, COMPONENTS) a, __DPV_MULTI_VECTOR(TYPE, COMPONENTS) b) { \
+	CALL_SITE_DEBUGLOC __DPV_MULTI_VECTOR(TYPE, COMPONENTS) dpv_min(__DPV_VECTOR(TYPE, COMPONENTS) a, __DPV_MULTI_VECTOR(TYPE, COMPONENTS) b) { \
 		return __DPV_FOREACH_COMPONENTS_CONSTRUCTOR(TYPE, COMPONENTS, __DPV_CODE_FUNCTION_2_ARGS, min, a, b); \
 	} \
-	__DPV_MULTI_VECTOR(TYPE, COMPONENTS) dpv_min(__DPV_MULTI_VECTOR(TYPE, COMPONENTS) a, __DPV_VECTOR(TYPE, COMPONENTS) b) { \
+	CALL_SITE_DEBUGLOC __DPV_MULTI_VECTOR(TYPE, COMPONENTS) dpv_min(__DPV_MULTI_VECTOR(TYPE, COMPONENTS) a, __DPV_VECTOR(TYPE, COMPONENTS) b) { \
 		return __DPV_FOREACH_COMPONENTS_CONSTRUCTOR(TYPE, COMPONENTS, __DPV_CODE_FUNCTION_2_ARGS, min, a, b); \
 	} \
-	__DPV_MULTI_VECTOR(TYPE, COMPONENTS) dpv_max(__DPV_MULTI_VECTOR(TYPE, COMPONENTS) a, __DPV_MULTI_VECTOR(TYPE, COMPONENTS) b) { \
+	CALL_SITE_DEBUGLOC __DPV_MULTI_VECTOR(TYPE, COMPONENTS) dpv_max(__DPV_MULTI_VECTOR(TYPE, COMPONENTS) a, __DPV_MULTI_VECTOR(TYPE, COMPONENTS) b) { \
 		return max(a, b); \
 	} \
-	__DPV_MULTI_VECTOR(TYPE, COMPONENTS) dpv_max(__DPV_VECTOR(TYPE, COMPONENTS) a, __DPV_MULTI_VECTOR(TYPE, COMPONENTS) b) { \
+	CALL_SITE_DEBUGLOC __DPV_MULTI_VECTOR(TYPE, COMPONENTS) dpv_max(__DPV_VECTOR(TYPE, COMPONENTS) a, __DPV_MULTI_VECTOR(TYPE, COMPONENTS) b) { \
 		return __DPV_FOREACH_COMPONENTS_CONSTRUCTOR(TYPE, COMPONENTS, __DPV_CODE_FUNCTION_2_ARGS, max, a, b); \
 	} \
-	__DPV_MULTI_VECTOR(TYPE, COMPONENTS) dpv_max(__DPV_MULTI_VECTOR(TYPE, COMPONENTS) a, __DPV_VECTOR(TYPE, COMPONENTS) b) { \
+	CALL_SITE_DEBUGLOC __DPV_MULTI_VECTOR(TYPE, COMPONENTS) dpv_max(__DPV_MULTI_VECTOR(TYPE, COMPONENTS) a, __DPV_VECTOR(TYPE, COMPONENTS) b) { \
 		return __DPV_FOREACH_COMPONENTS_CONSTRUCTOR(TYPE, COMPONENTS, __DPV_CODE_FUNCTION_2_ARGS, max, a, b); \
 	} \
 
@@ -412,13 +412,13 @@ DEFINE_DPV_FUNCTION_FOR_ALL_VECTORS(DEFINE_DPV_FUNCTION_VECTOR)
  *		half2 Distance = dpv_length(WorldPositionDelta);
  */
 #define DEFINE_DPV_FUNCTION_VECTOR(TYPE, COMPONENTS) \
-	__DPV_MULTI_SCALAR(TYPE) dpv_length2(__DPV_MULTI_VECTOR(TYPE, COMPONENTS) v) { \
+	CALL_SITE_DEBUGLOC __DPV_MULTI_SCALAR(TYPE) dpv_length2(__DPV_MULTI_VECTOR(TYPE, COMPONENTS) v) { \
 		return dpv_dot(v, v); \
 	} \
-	__DPV_MULTI_SCALAR(TYPE) dpv_length(__DPV_MULTI_VECTOR(TYPE, COMPONENTS) v) { \
+	CALL_SITE_DEBUGLOC __DPV_MULTI_SCALAR(TYPE) dpv_length(__DPV_MULTI_VECTOR(TYPE, COMPONENTS) v) { \
 		return sqrt(dpv_length2(v)); \
 	} \
-	__DPV_MULTI_VECTOR(TYPE, COMPONENTS) dpv_normalize(__DPV_MULTI_VECTOR(TYPE, COMPONENTS) v) { \
+	CALL_SITE_DEBUGLOC __DPV_MULTI_VECTOR(TYPE, COMPONENTS) dpv_normalize(__DPV_MULTI_VECTOR(TYPE, COMPONENTS) v) { \
 		return dpv_scale(v, rsqrt(dpv_length2(v))); \
 	} \
 
