Fix shader compilation errors due to CALL_SITE_DEBUGLOC

diff --git a/Engine/Shaders/Public/Platform.ush b/Engine/Shaders/Public/Platform.ush
--- a/Engine/Shaders/Public/Platform.ush
+++ b/Engine/Shaders/Public/Platform.ush
@@ -609,6 +609,10 @@ void ShaderYield()
 #define HOIST_DESCRIPTORS
 #endif
 
+#ifndef CALL_SITE_DEBUGLOC
+#define CALL_SITE_DEBUGLOC
+#endif
+
 #ifndef SCHEDULER_MIN_PRESSURE
 #define SCHEDULER_MIN_PRESSURE
 #endif
