Fix a bug that the GlobalClipPlane doesn't work on depth only vertex shader on mobile preview shader platform.

diff --git a/Engine/Shaders/Private/Common.ush b/Engine/Shaders/Private/Common.ush
--- a/Engine/Shaders/Private/Common.ush
+++ b/Engine/Shaders/Private/Common.ush
@@ -150,7 +150,7 @@ float4 FmodFloor(float4 Lhs, float4 Rhs)
 #define USE_EDITOR_SHADERS (PLATFORM_SUPPORTS_EDITOR_SHADERS && USE_DEVELOPMENT_SHADERS)
 
 // Using SV_ClipDistance has overhead (15% slower base pass in triangle bound test scene on PS4) so projects have to opt-in
-#define USE_GLOBAL_CLIP_PLANE (PLATFORM_SUPPORTS_GLOBAL_CLIP_PLANE && PROJECT_ALLOW_GLOBAL_CLIP_PLANE && !MATERIAL_DOMAIN_POSTPROCESS && !MATERIAL_DOMAIN_UI)
+#define USE_GLOBAL_CLIP_PLANE (PROJECT_ALLOW_GLOBAL_CLIP_PLANE && !MATERIAL_DOMAIN_POSTPROCESS && !MATERIAL_DOMAIN_UI)
 
 #if RAYHITGROUPSHADER || RAYMISSHADER
 
diff --git a/Engine/Shaders/Private/MobileBasePassCommon.ush b/Engine/Shaders/Private/MobileBasePassCommon.ush
--- a/Engine/Shaders/Private/MobileBasePassCommon.ush
+++ b/Engine/Shaders/Private/MobileBasePassCommon.ush
@@ -7,8 +7,6 @@
 #define FogStruct MobileBasePass.Fog
 #define PlanarReflectionStruct MobileBasePass.PlanarReflection
 
-#define USE_PS_CLIP_PLANE PROJECT_ALLOW_GLOBAL_CLIP_PLANE
-
 #undef NEEDS_LIGHTMAP_COORDINATE
 
 #define NEEDS_LIGHTMAP_COORDINATE		(LQ_TEXTURE_LIGHTMAP)
@@ -42,7 +40,7 @@ struct FSharedMobileBasePassInterpolants
 	float3 PixelPositionExcludingWPO	: TEXCOORD9;
 #endif
 
-#if USE_PS_CLIP_PLANE
+#if USE_GLOBAL_CLIP_PLANE
 	float OutClipDistance : OUTCLIPDIST;
 #endif
 
diff --git a/Engine/Shaders/Private/MobileBasePassPixelShader.usf b/Engine/Shaders/Private/MobileBasePassPixelShader.usf
--- a/Engine/Shaders/Private/MobileBasePassPixelShader.usf
+++ b/Engine/Shaders/Private/MobileBasePassPixelShader.usf
@@ -610,7 +610,7 @@ void Main(
 	ResolvedView = ResolveView();
 #endif
 
-#if USE_PS_CLIP_PLANE
+#if USE_GLOBAL_CLIP_PLANE
 	clip(BasePassInterpolants.OutClipDistance);
 #endif
 
diff --git a/Engine/Shaders/Private/MobileBasePassVertexShader.usf b/Engine/Shaders/Private/MobileBasePassVertexShader.usf
--- a/Engine/Shaders/Private/MobileBasePassVertexShader.usf
+++ b/Engine/Shaders/Private/MobileBasePassVertexShader.usf
@@ -93,7 +93,7 @@ void Main(
 	Output.BasePassInterpolants.PixelPositionExcludingWPO = WorldPositionExcludingWPO.xyz;
 #endif
 
-#if USE_PS_CLIP_PLANE
+#if USE_GLOBAL_CLIP_PLANE
 	Output.BasePassInterpolants.OutClipDistance = dot(ResolvedView.GlobalClippingPlane, float4(WorldPosition.xyz, 1));
 #endif
 
diff --git a/Engine/Shaders/Public/Platform.ush b/Engine/Shaders/Public/Platform.ush
--- a/Engine/Shaders/Public/Platform.ush
+++ b/Engine/Shaders/Public/Platform.ush
@@ -355,12 +355,6 @@
 #define SUPPORTS_INDEPENDENT_SAMPLERS (SM6_PROFILE || SM5_PROFILE || METAL_MRT_PROFILE || METAL_SM5_PROFILE || VULKAN_PROFILE_SM5 || VULKAN_PROFILE)
 #endif
 
-// Whether the platform supports a global clip plane through SV_ClipDistance
-// Ideally this would just be enabled for all SM4 and above feature level platforms, but not tested everywhere yet
-#ifndef PLATFORM_SUPPORTS_GLOBAL_CLIP_PLANE
-#define PLATFORM_SUPPORTS_GLOBAL_CLIP_PLANE (PS4_PROFILE || SM6_PROFILE || SM5_PROFILE || METAL_PROFILE || METAL_MRT_PROFILE || METAL_SM5_PROFILE || VULKAN_PROFILE_SM5)
-#endif
-
 // Whether the platform support pixel coverage on MSAA targets (SV_Coverage).
 #define SUPPORTS_PIXEL_COVERAGE (FEATURE_LEVEL >= FEATURE_LEVEL_SM5 && !COMPILER_GLSL && !MOBILE_EMULATION)
 
