Bump global shader version to fix crash due to stale parameter structures.

#rnx

diff --git a/Engine/Shaders/Public/ShaderVersion.ush b/Engine/Shaders/Public/ShaderVersion.ush
--- a/Engine/Shaders/Public/ShaderVersion.ush
+++ b/Engine/Shaders/Public/ShaderVersion.ush
@@ -2,4 +2,5 @@
 // Each time the console command is executed it generates a new GUID. As this file is included
 // in Platform.ush (which should be included in any shader) it allows to invalidate the shader DDC.
 // 
-// GUID = 26F2FBAA-7CF7-4B6A-9B2A-62D818DC9A01
+// GUID = B5734ACE-E057-480F-8E1E-67E08B1AB26D
+
