Adding initial GetResourceFromHeap and GetSamplerFromHeap support.

diff --git a/Engine/Config/Windows/DataDrivenPlatformInfo.ini b/Engine/Config/Windows/DataDrivenPlatformInfo.ini
--- a/Engine/Config/Windows/DataDrivenPlatformInfo.ini
+++ b/Engine/Config/Windows/DataDrivenPlatformInfo.ini
@@ -134,7 +134,7 @@ bSupportsCapsuleShadows=true
 bSupportsPercentageCloserShadows=true
 bSupportsVolumetricFog=true
 bSupportsIndexBufferUAVs=true
-; TODO: SM6.6
+bSupportsBindless=true
 
 [PreviewPlatform PCD3D_SM6]
 PlatformName=PCD3D_SM6
diff --git a/Engine/Shaders/Public/Platform/D3D/D3DCommon.ush b/Engine/Shaders/Public/Platform/D3D/D3DCommon.ush
--- a/Engine/Shaders/Public/Platform/D3D/D3DCommon.ush
+++ b/Engine/Shaders/Public/Platform/D3D/D3DCommon.ush
@@ -116,3 +116,9 @@ void UEReportAssertWithPayload(uint AssertID, uint4 Payload)
 		}
 	#endif
 #endif
+
+#if PLATFORM_SUPPORTS_BINDLESS
+	// SM6.6 path
+	#define GetResourceFromHeap(Type, Index) ResourceDescriptorHeap[Index]
+	#define GetSamplerFromHeap(Type, Index) SamplerDescriptorHeap[Index]
+#endif
