Fixed: Meshes appear to tear/stretch on Android (S21 Mali, Vulkan)

diff --git a/Engine/Shaders/Public/Platform.ush b/Engine/Shaders/Public/Platform.ush
--- a/Engine/Shaders/Public/Platform.ush
+++ b/Engine/Shaders/Public/Platform.ush
@@ -330,6 +330,11 @@
 
 // ---------------------------------------------------- Platform dependent supports
 
+// On mobile use interpolator for ClipDistance as it does not work/supported on all mobile platforms
+#if (FEATURE_LEVEL == FEATURE_LEVEL_ES3_1)
+#define SV_ClipDistance OUTCLIPDIST
+#endif
+
 // non-editor platforms generally never want development/editor features.
 #ifndef PLATFORM_SUPPORTS_DEVELOPMENT_SHADERS
 #define PLATFORM_SUPPORTS_DEVELOPMENT_SHADERS 1
diff --git a/Engine/Shaders/Public/Platform/GL/GLCommon.ush b/Engine/Shaders/Public/Platform/GL/GLCommon.ush
--- a/Engine/Shaders/Public/Platform/GL/GLCommon.ush
+++ b/Engine/Shaders/Public/Platform/GL/GLCommon.ush
@@ -10,5 +10,3 @@
 #pragma once
 
 
-// GLES does not support gl_ClipDistance, use interpolator for it
-#define SV_ClipDistance OUTCLIPDIST
\ No newline at end of file
