Fix for mobile deferred shading in Metal RHI

diff --git a/Engine/Shaders/Public/Platform/Metal/MetalSubpassSupport.ush b/Engine/Shaders/Public/Platform/Metal/MetalSubpassSupport.ush
--- a/Engine/Shaders/Public/Platform/Metal/MetalSubpassSupport.ush
+++ b/Engine/Shaders/Public/Platform/Metal/MetalSubpassSupport.ush
@@ -38,14 +38,14 @@
 	Texture2D<float> gl_LastFragDataR_0;
 	#define SubpassFetchR_0() gl_LastFragDataR_0.Load(uint3(0, 0, 0), 0)
 
-	Texture2D<float> gl_LastFragDataR_1;
-	#define SubpassFetchR_1() gl_LastFragDataR_1.Load(uint3(0, 0, 0), 0)
+    // R_1 in MetalCommon.ush
 
 	Texture2D<float> gl_LastFragDataR_2;
 	#define SubpassFetchR_2() gl_LastFragDataR_2.Load(uint3(0, 0, 0), 0)
 
 	Texture2D<float> gl_LastFragDataR_3;
 	#define SubpassFetchR_3() gl_LastFragDataR_3.Load(uint3(0, 0, 0), 0)
+    
 	// R_4 in MetalCommon.ush
 
 	Texture2D<float> gl_LastFragDataR_5;
