Works arround a bunch of HLSL2021 bugs

DXC is missing some override of select()
DXC's -spirv is missing and() and or() intrinsics
Fixes CFLAG_WarningsAsErrors that would cause a compilation failure if rewritted HLSL for FXC had warning
Adds a COMPILER_SUPPORTS_HLSL2021

diff --git a/Engine/Shaders/Public/Platform.ush b/Engine/Shaders/Public/Platform.ush
--- a/Engine/Shaders/Public/Platform.ush
+++ b/Engine/Shaders/Public/Platform.ush
@@ -187,6 +187,11 @@
 #define PLATFORM_SUPPORTS_EDITOR_SHADERS 1
 #endif
 
+#ifndef COMPILER_SUPPORTS_HLSL2021
+#define COMPILER_SUPPORTS_HLSL2021 0
+#endif
+
+
 // ---------------------------------------------------- Alternative floating point types
 
 #ifndef FORCE_FLOATS
@@ -822,6 +827,62 @@ half4 max3( half4 a, half4 b, half4 c )
 
 #endif
 
+// https://devblogs.microsoft.com/directx/announcing-hlsl-2021/
+// HLSL 2021 supports Logical Operator Short Circuiting. To do vector bool operations, need to use and() or() select()
+// Sadly the HLSL2021 standard does not overload select() very well...
+#define select(cond,a,b) select_internal(cond,a,b)
+#define DEFINE_SELECT(TYPE) \
+	TYPE    select_internal(bool    c, TYPE    a, TYPE    b) { return TYPE   (c   ? a.x : b.x); } \
+	\
+	TYPE##2 select_internal(bool    c, TYPE    a, TYPE##2 b) { return TYPE##2(c   ? a   : b.x, c   ? a   : b.y); } \
+	TYPE##2 select_internal(bool    c, TYPE##2 a, TYPE    b) { return TYPE##2(c   ? a.x : b  , c   ? a.y : b  ); } \
+	TYPE##2 select_internal(bool    c, TYPE##2 a, TYPE##2 b) { return TYPE##2(c   ? a.x : b.x, c   ? a.y : b.y); } \
+	TYPE##2 select_internal(bool##2 c, TYPE    a, TYPE    b) { return TYPE##2(c.x ? a   : b  , c.y ? a   : b  ); } \
+	TYPE##2 select_internal(bool##2 c, TYPE    a, TYPE##2 b) { return TYPE##2(c.x ? a   : b.x, c.y ? a   : b.y); } \
+	TYPE##2 select_internal(bool##2 c, TYPE##2 a, TYPE    b) { return TYPE##2(c.x ? a.x : b  , c.y ? a.y : b  ); } \
+	TYPE##2 select_internal(bool##2 c, TYPE##2 a, TYPE##2 b) { return TYPE##2(c.x ? a.x : b.x, c.y ? a.y : b.y); } \
+	\
+	TYPE##3 select_internal(bool    c, TYPE    a, TYPE##3 b) { return TYPE##3(c   ? a   : b.x, c   ? a   : b.y, c   ? a   : b.z); } \
+	TYPE##3 select_internal(bool    c, TYPE##3 a, TYPE    b) { return TYPE##3(c   ? a.x : b  , c   ? a.y : b  , c   ? a.z : b  ); } \
+	TYPE##3 select_internal(bool    c, TYPE##3 a, TYPE##3 b) { return TYPE##3(c   ? a.x : b.x, c   ? a.y : b.y, c   ? a.z : b.z); } \
+	TYPE##3 select_internal(bool##3 c, TYPE    a, TYPE    b) { return TYPE##3(c.x ? a   : b  , c.y ? a   : b  , c.z ? a   : b  ); } \
+	TYPE##3 select_internal(bool##3 c, TYPE    a, TYPE##3 b) { return TYPE##3(c.x ? a   : b.x, c.y ? a   : b.y, c.z ? a   : b.z); } \
+	TYPE##3 select_internal(bool##3 c, TYPE##3 a, TYPE    b) { return TYPE##3(c.x ? a.x : b  , c.y ? a.y : b  , c.z ? a.z : b  ); } \
+	TYPE##3 select_internal(bool##3 c, TYPE##3 a, TYPE##3 b) { return TYPE##3(c.x ? a.x : b.x, c.y ? a.y : b.y, c.z ? a.z : b.z); } \
+	\
+	TYPE##4 select_internal(bool    c, TYPE    a, TYPE##4 b) { return TYPE##4(c   ? a   : b.x, c   ? a   : b.y, c   ? a   : b.z, c   ? a   : b.w); } \
+	TYPE##4 select_internal(bool    c, TYPE##4 a, TYPE    b) { return TYPE##4(c   ? a.x : b  , c   ? a.y : b  , c   ? a.z : b  , c   ? a.w : b  ); } \
+	TYPE##4 select_internal(bool    c, TYPE##4 a, TYPE##4 b) { return TYPE##4(c   ? a.x : b.x, c   ? a.y : b.y, c   ? a.z : b.z, c   ? a.w : b.w); } \
+	TYPE##4 select_internal(bool##4 c, TYPE    a, TYPE    b) { return TYPE##4(c.x ? a   : b  , c.y ? a   : b  , c.z ? a   : b  , c.w ? a   : b  ); } \
+	TYPE##4 select_internal(bool##4 c, TYPE    a, TYPE##4 b) { return TYPE##4(c.x ? a   : b.x, c.y ? a   : b.y, c.z ? a   : b.z, c.w ? a   : b.w); } \
+	TYPE##4 select_internal(bool##4 c, TYPE##4 a, TYPE    b) { return TYPE##4(c.x ? a.x : b  , c.y ? a.y : b  , c.z ? a.z : b  , c.w ? a.w : b  ); } \
+	TYPE##4 select_internal(bool##4 c, TYPE##4 a, TYPE##4 b) { return TYPE##4(c.x ? a.x : b.x, c.y ? a.y : b.y, c.z ? a.z : b.z, c.w ? a.w : b.w); } \
+
+DEFINE_SELECT(bool)
+DEFINE_SELECT(uint)
+DEFINE_SELECT(int)
+DEFINE_SELECT(float)
+#if PLATFORM_SUPPORTS_REAL_TYPES
+DEFINE_SELECT(half)
+DEFINE_SELECT(uint16_t)
+DEFINE_SELECT(int16_t)
+#endif
+#undef DEFINE_SELECT
+
+// Works around bug in the spirv for the missing implementation of the and() and or() intrinsics.
+bool  and_internal(bool  a, bool  b) { return bool(a && b); }
+bool2 and_internal(bool2 a, bool2 b) { return bool2(a.x && b.x, a.y && b.y); }
+bool3 and_internal(bool3 a, bool3 b) { return bool3(a.x && b.x, a.y && b.y, a.z && b.z); }
+bool4 and_internal(bool4 a, bool4 b) { return bool4(a.x && b.x, a.y && b.y, a.z && b.z, a.w && b.w); }
+
+bool  or_internal(bool  a, bool  b) { return bool(a || b); }
+bool2 or_internal(bool2 a, bool2 b) { return bool2(a.x || b.x, a.y || b.y); }
+bool3 or_internal(bool3 a, bool3 b) { return bool3(a.x || b.x, a.y || b.y, a.z || b.z); }
+bool4 or_internal(bool4 a, bool4 b) { return bool4(a.x || b.x, a.y || b.y, a.z || b.z, a.w || b.w); }
+
+#define and(a, b) and_internal(a, b)
+#define or(a, b) or_internal(a, b)
+
 #if PLATFORM_SUPPORTS_REAL_TYPES && !defined(COMPILER_SUPPORTS_PACK_B32_B16)
 
 // Function that explicitly use RDNA's v_pack_b32_f16 on supported platform. Note that RDNA's documentation call this instruction v_pack_b32_f16
diff --git a/Engine/Source/Developer/Windows/ShaderFormatD3D/Private/D3DShaderCompiler.cpp b/Engine/Source/Developer/Windows/ShaderFormatD3D/Private/D3DShaderCompiler.cpp
--- a/Engine/Source/Developer/Windows/ShaderFormatD3D/Private/D3DShaderCompiler.cpp
+++ b/Engine/Source/Developer/Windows/ShaderFormatD3D/Private/D3DShaderCompiler.cpp
@@ -766,7 +766,7 @@ bool CompileAndProcessD3DShaderFXC(FString& PreprocessedShaderSource, const FStr
 				/*pInclude=*/ NULL,
 				"main",
 				TCHAR_TO_ANSI(ShaderProfile),
-				CompileFlags,
+				CompileFlags & (~D3DCOMPILE_WARNINGS_ARE_ERRORS),
 				0,
 				Shader.GetInitReference(),
 				Errors.GetInitReference(),
diff --git a/Engine/Source/Runtime/Engine/Private/ShaderCompiler/ShaderCompiler.cpp b/Engine/Source/Runtime/Engine/Private/ShaderCompiler/ShaderCompiler.cpp
--- a/Engine/Source/Runtime/Engine/Private/ShaderCompiler/ShaderCompiler.cpp
+++ b/Engine/Source/Runtime/Engine/Private/ShaderCompiler/ShaderCompiler.cpp
@@ -5727,6 +5727,11 @@ void GlobalBeginCompileShader(
 
 	Input.Environment.SetDefine(TEXT("HAS_INVERTED_Z_BUFFER"), (bool)ERHIZBuffer::IsInverted);
 
+	if (Input.Environment.CompilerFlags.Contains(CFLAG_HLSL2021))
+	{
+		Input.Environment.SetDefine(TEXT("COMPILER_SUPPORTS_HLSL2021"), 1);
+	}
+
 	{
 		static const auto CVar = IConsoleManager::Get().FindTConsoleVariableDataInt(TEXT("r.ClearCoatNormal"));
 		Input.Environment.SetDefine(TEXT("CLEAR_COAT_BOTTOM_NORMAL"), CVar ? (CVar->GetValueOnAnyThread() != 0) && !bIsMobilePlatform : 0);
