Enable INVARIANT and fix MakePrecise() for D3D shaders.
Fixes issues with DXC compiler where we need precise keyword to ensure depth and base pass vertex outputs match.

diff --git a/Engine/Shaders/Public/Platform.ush b/Engine/Shaders/Public/Platform.ush
--- a/Engine/Shaders/Public/Platform.ush
+++ b/Engine/Shaders/Public/Platform.ush
@@ -542,7 +542,7 @@ void ShaderYield()
 #endif
 
 #ifndef INVARIANT
-#define INVARIANT
+#define INVARIANT(X) X
 #endif
 
 #ifndef ENABLE_RE_Z
diff --git a/Engine/Shaders/Public/Platform/D3D/D3DCommon.ush b/Engine/Shaders/Public/Platform/D3D/D3DCommon.ush
--- a/Engine/Shaders/Public/Platform/D3D/D3DCommon.ush
+++ b/Engine/Shaders/Public/Platform/D3D/D3DCommon.ush
@@ -7,11 +7,11 @@ D3DCommon.usf: Common defines for the d3d shader compiler
 // Update this GUID to invalidate D3D shaders and force a shader recompilation for D3D RHIs.
 // GUID: B8728C4B-C1E8-45AD-B3E4-A6761FAFA247
 
-precise float MakePrecise(precise float v) { return v; }
-precise float2 MakePrecise(precise float2 v) { return v; }
-precise float3 MakePrecise(precise float3 v) { return v; }
-precise float4 MakePrecise(precise float4 v) { return v; }
-#define INVARIANT(X) X
+precise float MakePrecise(float v) { precise float pv = v; return pv; }
+precise float2 MakePrecise(float2 v) { precise float2 pv = v; return pv; }
+precise float3 MakePrecise(float3 v) { precise float3 pv = v; return pv; }
+precise float4 MakePrecise(float4 v) { precise float4 pv = v; return pv; }
+#define INVARIANT(X) MakePrecise(X)
 
 #if PLATFORM_SUPPORTS_STATIC_SAMPLERS
 
