Fix compilation error on platforms which don't support ROV.

Despite the guard/default definition for non-supporting platforms, if a uniform buffer include an ROV resource, it causes an error.

diff --git a/Engine/Shaders/Public/Platform.ush b/Engine/Shaders/Public/Platform.ush
--- a/Engine/Shaders/Public/Platform.ush
+++ b/Engine/Shaders/Public/Platform.ush
@@ -94,6 +94,7 @@
 
 #ifndef PLATFORM_SUPPORTS_ROV
 #define PLATFORM_SUPPORTS_ROV 0
+#define RasterizerOrderedTexture2D RWTexture2D
 #endif
 
 #ifndef SM6_PROFILE
