Fix bug with INVARIANT macro - lack of parenthesis breaks order of operations if non-trivial math is being performed

diff --git a/Engine/Shaders/Public/Platform.ush b/Engine/Shaders/Public/Platform.ush
--- a/Engine/Shaders/Public/Platform.ush
+++ b/Engine/Shaders/Public/Platform.ush
@@ -551,7 +551,7 @@ void ShaderYield()
 #endif
 
 #ifndef INVARIANT
-#define INVARIANT(X) X
+#define INVARIANT(X) (X)
 #endif
 
 #ifndef ENABLE_RE_Z
