VulkanRHI:  improve GPU rendering performance with earlydepthstencil

diff --git a/Engine/Shaders/Public/Platform.ush b/Engine/Shaders/Public/Platform.ush
--- a/Engine/Shaders/Public/Platform.ush
+++ b/Engine/Shaders/Public/Platform.ush
@@ -525,7 +525,7 @@ void ShaderYield()
 
 #endif // SM6_PROFILE || SM5_PROFILE || COMPILER_SUPPORTS_ATTRIBUTES
 
-#if SM6_PROFILE || SM5_PROFILE || METAL_MRT_PROFILE || METAL_SM5_PROFILE || ES3_1_PROFILE
+#if SM6_PROFILE || SM5_PROFILE || METAL_MRT_PROFILE || METAL_SM5_PROFILE || ES3_1_PROFILE || VULKAN_PROFILE_SM5
 #define EARLYDEPTHSTENCIL [earlydepthstencil]
 #endif
 
