Adds support for RDNA's ds_swizzle_b32 lane rotation mode

diff --git a/Engine/Shaders/Public/WaveBroadcastIntrinsics.ush b/Engine/Shaders/Public/WaveBroadcastIntrinsics.ush
--- a/Engine/Shaders/Public/WaveBroadcastIntrinsics.ush
+++ b/Engine/Shaders/Public/WaveBroadcastIntrinsics.ush
@@ -10,7 +10,17 @@
 #include "Platform.ush"
 
 
-#define PLATFORM_SUPPORTS_WAVE_BROADCAST (COMPILER_SUPPORTS_WAVE_SWIZZLE_GCN)
+#ifdef COMPILER_SUPPORTS_WAVE_SWIZZLE_GCN
+	#define PLATFORM_SUPPORTS_WAVE_BROADCAST (COMPILER_SUPPORTS_WAVE_SWIZZLE_GCN)
+#else
+	#define PLATFORM_SUPPORTS_WAVE_BROADCAST 0
+#endif
+
+#ifdef COMPILER_SUPPORTS_WAVE_SWIZZLE_RDNA
+	#define PLATFORM_SUPPORTS_WAVE_ROTATE (COMPILER_SUPPORTS_WAVE_SWIZZLE_RDNA)
+#else
+	#define PLATFORM_SUPPORTS_WAVE_ROTATE 0
+#endif
 
 
 /** Different instruction available for WaveBroadcast() */
@@ -18,6 +28,9 @@
 	#if PLATFORM_SUPPORTS_WAVE_BROADCAST
 		#define WAVE_BROADCAST_GCN_SWIZZLE 1
 	#endif
+	#if PLATFORM_SUPPORTS_WAVE_ROTATE
+		#define WAVE_BROADCAST_RDNA_ROTATE 2
+	#endif
 
 /** Compile time structure to choose which broadcasting should be done. */
 struct FWaveBroadcastSettings
@@ -25,13 +38,22 @@ struct FWaveBroadcastSettings
 	// Broadcast operations.
 	uint Operation;
 
-	// Informations about the broadcast.
-	uint SwizzleAnd;
-	uint SwizzleOr;
-	uint SwizzleXor;
+	// Informations about the broadcast for GCN's ds_swizzle.
+	#if PLATFORM_SUPPORTS_WAVE_BROADCAST
+		uint SwizzleAnd;
+		uint SwizzleOr;
+		uint SwizzleXor;
+	#endif
+	
+	// Information about the broadcast for RDNA's ds_swizzle's rotate functionality
+	#if PLATFORM_SUPPORTS_WAVE_ROTATE
+		int Rotate;
+		uint RotateFixMask;
+	#endif
 };
 
 /** Generic broadcast instruction. */
+CALL_SITE_DEBUGLOC
 uint WaveBroadcast(const FWaveBroadcastSettings Settings, uint x)
 {
 	if (0)
@@ -44,20 +66,17 @@ uint WaveBroadcast(const FWaveBroadcastSettings Settings, uint x)
 		return WaveLaneSwizzleGCN(x, Settings.SwizzleAnd, Settings.SwizzleOr, Settings.SwizzleXor);
 	}
 	#endif
+	#if PLATFORM_SUPPORTS_WAVE_ROTATE
+	else if (Settings.Operation == WAVE_BROADCAST_RDNA_ROTATE)
+	{
+		return WaveLaneRotateSwizzleRDNA(x, /* rotate_amount = */ Settings.Rotate, Settings.RotateFixMask);
+	}
+	#endif
 
 	return x;
 }
 
-int WaveBroadcast(const FWaveBroadcastSettings Settings, int x)
-{
-	return asint(WaveBroadcast(Settings, asuint(x)));
-} 
-
-float WaveBroadcast(const FWaveBroadcastSettings Settings, float x)
-{
-	return asfloat(WaveBroadcast(Settings, asuint(x)));
-}
-
+CALL_SITE_DEBUGLOC
 uint2 WaveBroadcast(const FWaveBroadcastSettings Settings, uint2 v)
 {
 	return uint2(
@@ -65,6 +84,7 @@ uint2 WaveBroadcast(const FWaveBroadcastSettings Settings, uint2 v)
 		WaveBroadcast(Settings, v.y));
 } 
 
+CALL_SITE_DEBUGLOC
 uint3 WaveBroadcast(const FWaveBroadcastSettings Settings, uint3 v)
 {
 	return uint3(
@@ -73,6 +93,7 @@ uint3 WaveBroadcast(const FWaveBroadcastSettings Settings, uint3 v)
 		WaveBroadcast(Settings, v.z));
 } 
 
+CALL_SITE_DEBUGLOC
 uint4 WaveBroadcast(const FWaveBroadcastSettings Settings, uint4 v)
 {
 	return uint4(
@@ -82,6 +103,14 @@ uint4 WaveBroadcast(const FWaveBroadcastSettings Settings, uint4 v)
 		WaveBroadcast(Settings, v.w));
 } 
 
+
+CALL_SITE_DEBUGLOC
+int WaveBroadcast(const FWaveBroadcastSettings Settings, int x)
+{
+	return asint(WaveBroadcast(Settings, asuint(x)));
+} 
+
+CALL_SITE_DEBUGLOC
 int2 WaveBroadcast(const FWaveBroadcastSettings Settings, int2 v)
 {
 	return int2(
@@ -89,6 +118,7 @@ int2 WaveBroadcast(const FWaveBroadcastSettings Settings, int2 v)
 		WaveBroadcast(Settings, v.y));
 } 
 
+CALL_SITE_DEBUGLOC
 int3 WaveBroadcast(const FWaveBroadcastSettings Settings, int3 v)
 {
 	return int3(
@@ -97,6 +127,7 @@ int3 WaveBroadcast(const FWaveBroadcastSettings Settings, int3 v)
 		WaveBroadcast(Settings, v.z));
 } 
 
+CALL_SITE_DEBUGLOC
 int4 WaveBroadcast(const FWaveBroadcastSettings Settings, int4 v)
 {
 	return int4(
@@ -106,6 +137,14 @@ int4 WaveBroadcast(const FWaveBroadcastSettings Settings, int4 v)
 		WaveBroadcast(Settings, v.w));
 } 
 
+
+CALL_SITE_DEBUGLOC
+float WaveBroadcast(const FWaveBroadcastSettings Settings, float x)
+{
+	return asfloat(WaveBroadcast(Settings, asuint(x)));
+}
+
+CALL_SITE_DEBUGLOC
 float2 WaveBroadcast(const FWaveBroadcastSettings Settings, float2 v)
 {
 	return float2(
@@ -113,14 +152,16 @@ float2 WaveBroadcast(const FWaveBroadcastSettings Settings, float2 v)
 		WaveBroadcast(Settings, v.y));
 } 
 
+CALL_SITE_DEBUGLOC
 float3 WaveBroadcast(const FWaveBroadcastSettings Settings, float3 v)
 {
 	return float3(
 		WaveBroadcast(Settings, v.x),
 		WaveBroadcast(Settings, v.y),
 		WaveBroadcast(Settings, v.z));
-} 
+}
 
+CALL_SITE_DEBUGLOC
 float4 WaveBroadcast(const FWaveBroadcastSettings Settings, float4 v)
 {
 	return float4(
@@ -131,15 +172,115 @@ float4 WaveBroadcast(const FWaveBroadcastSettings Settings, float4 v)
 } 
 
 
+#if PLATFORM_SUPPORTS_REAL_TYPES
+
+CALL_SITE_DEBUGLOC
+uint16_t WaveBroadcast(const FWaveBroadcastSettings Settings, uint16_t x)
+{
+	return uint16_t(WaveBroadcast(Settings, uint(x)));
+}
+
+CALL_SITE_DEBUGLOC
+uint16_t2 WaveBroadcast(const FWaveBroadcastSettings Settings, uint16_t2 v)
+#if COMPILER_SUPPORT_UINT16_BITCAST
+{
+	uint vu = bit_cast_uint(v);
+	uint bvu = WaveBroadcast(Settings, vu);
+	return bit_cast_uint16_t2(bvu);
+} 
+#else
+{
+	return uint16_t2(
+		WaveBroadcast(Settings, v.x),
+		WaveBroadcast(Settings, v.y));
+}
+#endif
+
+CALL_SITE_DEBUGLOC
+uint16_t3 WaveBroadcast(const FWaveBroadcastSettings Settings, uint16_t3 v)
+{
+	return uint16_t3(
+		WaveBroadcast(Settings, v.xy),
+		WaveBroadcast(Settings, v.z));
+}
+
+CALL_SITE_DEBUGLOC
+uint16_t4 WaveBroadcast(const FWaveBroadcastSettings Settings, uint16_t4 v)
+{
+	return uint16_t4(
+		WaveBroadcast(Settings, v.xy),
+		WaveBroadcast(Settings, v.zw));
+} 
+
+
+CALL_SITE_DEBUGLOC
+int16_t WaveBroadcast(const FWaveBroadcastSettings Settings, int16_t x)
+{
+	return asint16(WaveBroadcast(Settings, asuint16(x)));
+}
+
+CALL_SITE_DEBUGLOC
+int16_t2 WaveBroadcast(const FWaveBroadcastSettings Settings, int16_t2 v)
+{
+	return asint16(WaveBroadcast(Settings, asuint16(v)));
+}
+
+CALL_SITE_DEBUGLOC
+int16_t3 WaveBroadcast(const FWaveBroadcastSettings Settings, int16_t3 v)
+{
+	return asint16(WaveBroadcast(Settings, asuint16(v)));
+}
+
+CALL_SITE_DEBUGLOC
+int16_t4 WaveBroadcast(const FWaveBroadcastSettings Settings, int16_t4 v)
+{
+	return asint16(WaveBroadcast(Settings, asuint16(v)));
+}
+
+
+CALL_SITE_DEBUGLOC
+half WaveBroadcast(const FWaveBroadcastSettings Settings, half x)
+{
+	return asfloat16(WaveBroadcast(Settings, asuint16(x)));
+}
+
+CALL_SITE_DEBUGLOC
+half2 WaveBroadcast(const FWaveBroadcastSettings Settings, half2 v)
+{
+	return asfloat16(WaveBroadcast(Settings, asuint16(v)));
+}
+
+CALL_SITE_DEBUGLOC
+half3 WaveBroadcast(const FWaveBroadcastSettings Settings, half3 v)
+{
+	return asfloat16(WaveBroadcast(Settings, asuint16(v)));
+}
+
+CALL_SITE_DEBUGLOC
+half4 WaveBroadcast(const FWaveBroadcastSettings Settings, half4 v)
+{
+	return asfloat16(WaveBroadcast(Settings, asuint16(v)));
+}
+
+#endif
+
+
 
 /** Init a nop broadcast */
+CALL_SITE_DEBUGLOC
 FWaveBroadcastSettings InitNopBroadcast()
 {
 	FWaveBroadcastSettings Settings;
-	Settings.Operation = WAVE_BROADCAST_NOP;
-	Settings.SwizzleAnd = 0x00;
-	Settings.SwizzleOr  = 0x00;
-	Settings.SwizzleXor = 0x00;
+	Settings.Operation  = WAVE_BROADCAST_NOP;
+	#if PLATFORM_SUPPORTS_WAVE_BROADCAST
+		Settings.SwizzleAnd = 0x00;
+		Settings.SwizzleOr  = 0x00;
+		Settings.SwizzleXor = 0x00;
+	#endif
+	#if PLATFORM_SUPPORTS_WAVE_ROTATE
+		Settings.Rotate        = +0;
+		Settings.RotateFixMask = 0x00;
+	#endif
 	return Settings;
 }
 
@@ -150,10 +291,11 @@ FWaveBroadcastSettings InitNopBroadcast()
  *
  *	return src[laneId ^ XorButterFly];
  */
+CALL_SITE_DEBUGLOC
 FWaveBroadcastSettings InitWaveXorButterfly(const uint XorButterFly)
 {
-	FWaveBroadcastSettings Settings;
-	Settings.Operation = WAVE_BROADCAST_GCN_SWIZZLE;
+	FWaveBroadcastSettings Settings = InitNopBroadcast();
+	Settings.Operation  = WAVE_BROADCAST_GCN_SWIZZLE;
 	Settings.SwizzleAnd = 0x1F;
 	Settings.SwizzleOr  = 0x00;
 	Settings.SwizzleXor = XorButterFly;
@@ -172,6 +314,7 @@ FWaveBroadcastSettings InitWaveXorButterfly(const uint XorButterFly)
  *
  *  return | 4  5  6  7| 0  1  2  3|12 13 ...
  */
+CALL_SITE_DEBUGLOC
 FWaveBroadcastSettings InitWaveSwapWithinLaneGroup(const uint LaneGroupSize)
 {
 	return InitWaveXorButterfly(/* XorButterFly = */ LaneGroupSize >> 1);
@@ -190,12 +333,13 @@ FWaveBroadcastSettings InitWaveSwapWithinLaneGroup(const uint LaneGroupSize)
  *
  *  return | 2  3  2  3  2  3  2  3|10 11 ...
  */
+CALL_SITE_DEBUGLOC
 FWaveBroadcastSettings InitWaveBroadcastLaneGroup(const uint LaneGroupSize, const uint InnerLaneGroupSize, const uint InnerLaneGroupId)
 {
 	const uint InnerGroupCount = LaneGroupSize / InnerLaneGroupSize;
-
-	FWaveBroadcastSettings Settings;
-	Settings.Operation = WAVE_BROADCAST_GCN_SWIZZLE;
+	
+	FWaveBroadcastSettings Settings = InitNopBroadcast();
+	Settings.Operation  = WAVE_BROADCAST_GCN_SWIZZLE;
 	Settings.SwizzleAnd = ~((InnerGroupCount - 1) * InnerLaneGroupSize);
 	Settings.SwizzleOr  = InnerLaneGroupId * InnerLaneGroupSize;
 	Settings.SwizzleXor = 0x00;
@@ -203,3 +347,32 @@ FWaveBroadcastSettings InitWaveBroadcastLaneGroup(const uint LaneGroupSize, cons
 }
 
 #endif // PLATFORM_SUPPORTS_WAVE_BROADCAST
+
+#if PLATFORM_SUPPORTS_WAVE_ROTATE
+
+/**  Rotate lane (-31 to +31) over a lane group (size is power of two in [2; 32]).
+ *
+ *  If a lane is not active, the VGPR value returned is 0.
+ *
+ *  LaneGroupSize = 8
+ *  LaneRotation = +3
+ *
+ *         |  lane group (size=8)  |
+ *     x = | 0  1  2  3  4  5  6  7| 8  9 ...
+ *
+ *  return | 3  4  5  6  7  0  1  2|11 12 ...
+ */
+CALL_SITE_DEBUGLOC
+FWaveBroadcastSettings InitWaveRotateLaneGroup(const uint LaneGroupSize, const int LaneRotation)
+{
+	const uint RotateMask    = LaneGroupSize - 1;
+	const uint RotateFixMask = (~RotateMask) & 0x1F;
+
+	FWaveBroadcastSettings Settings;
+	Settings.Operation     = WAVE_BROADCAST_RDNA_ROTATE;
+	Settings.Rotate        = LaneRotation;
+	Settings.RotateFixMask = RotateFixMask;
+	return Settings;
+}
+
+#endif // PLATFORM_SUPPORTS_WAVE_ROTATE
