Bump the shader version. This is necessary due to CL 20480382.

diff --git a/Engine/Shaders/Public/ShaderVersion.ush b/Engine/Shaders/Public/ShaderVersion.ush
--- a/Engine/Shaders/Public/ShaderVersion.ush
+++ b/Engine/Shaders/Public/ShaderVersion.ush
@@ -1,5 +1,5 @@
-// This file is automatically generated by the console command r.InvalidateCachedShaders
-// Each time the console command is executed it generates a new GUID. As this file is included
-// in Platform.ush (which should be included in any shader) it allows to invalidate the shader DDC.
-// 
-// GUID = 3FBEFF0E-9091-4DE4-9A24-D61106A33EDA
+// This file is automatically generated by the console command r.InvalidateCachedShaders
+// Each time the console command is executed it generates a new GUID. As this file is included
+// in Platform.ush (which should be included in any shader) it allows to invalidate the shader DDC.
+// 
+// GUID = 9F6D0575-8F17-4535-BEF8-606365FE370B
