Fixes wave broadcast intrinsics on vulkan

diff --git a/Engine/Shaders/Public/WaveBroadcastIntrinsics.ush b/Engine/Shaders/Public/WaveBroadcastIntrinsics.ush
--- a/Engine/Shaders/Public/WaveBroadcastIntrinsics.ush
+++ b/Engine/Shaders/Public/WaveBroadcastIntrinsics.ush
@@ -12,7 +12,7 @@
 
 #if COMPILER_SUPPORTS_WAVE_PERMUTE
 	#define PLATFORM_SUPPORTS_WAVE_READ_AT 1
-#elif FEATURE_LEVEL >= FEATURE_LEVEL_SM6
+#elif FEATURE_LEVEL >= FEATURE_LEVEL_SM6 || PLATFORM_SUPPORTS_SM6_0_WAVE_OPERATIONS
 	#define PLATFORM_SUPPORTS_WAVE_READ_AT 1
 #else
 	#define PLATFORM_SUPPORTS_WAVE_READ_AT 0
@@ -20,7 +20,7 @@
 
 #if COMPILER_SUPPORTS_WAVE_SWIZZLE_GCN
 	#define PLATFORM_SUPPORTS_WAVE_BROADCAST 1
-#elif FEATURE_LEVEL >= FEATURE_LEVEL_SM6
+#elif FEATURE_LEVEL >= FEATURE_LEVEL_SM6 || PLATFORM_SUPPORTS_SM6_0_WAVE_OPERATIONS
 	#define PLATFORM_SUPPORTS_WAVE_BROADCAST 1
 #else
 	#define PLATFORM_SUPPORTS_WAVE_BROADCAST 0
@@ -28,7 +28,7 @@
 
 #if defined(COMPILER_SUPPORTS_WAVE_SWIZZLE_RDNA) && COMPILER_SUPPORTS_WAVE_SWIZZLE_RDNA
 	#define PLATFORM_SUPPORTS_WAVE_ROTATE 1
-#elif FEATURE_LEVEL >= FEATURE_LEVEL_SM6
+#elif FEATURE_LEVEL >= FEATURE_LEVEL_SM6 || PLATFORM_SUPPORTS_SM6_0_WAVE_OPERATIONS
 	#define PLATFORM_SUPPORTS_WAVE_ROTATE 1
 #else
 	#define PLATFORM_SUPPORTS_WAVE_ROTATE 0
@@ -112,7 +112,7 @@ uint WaveBroadcast(const FWaveBroadcastSettings Settings, uint x)
 		{
 			return WaveReadLaneAtVarying(x, Settings.SourceLaneIndex);
 		}
-		#elif FEATURE_LEVEL >= FEATURE_LEVEL_SM6
+		#elif FEATURE_LEVEL >= FEATURE_LEVEL_SM6 || PLATFORM_SUPPORTS_SM6_0_WAVE_OPERATIONS
 		{
 			return WaveReadLaneAt(x, Settings.SourceLaneIndex);
 		}
@@ -142,7 +142,7 @@ uint WaveBroadcast(const FWaveBroadcastSettings Settings, uint x)
 		{
 			return WaveLaneSwizzleGCN(x, Settings.SwizzleAnd, Settings.SwizzleOr, Settings.SwizzleXor);
 		}
-		#elif FEATURE_LEVEL >= FEATURE_LEVEL_SM6
+		#elif FEATURE_LEVEL >= FEATURE_LEVEL_SM6 || PLATFORM_SUPPORTS_SM6_0_WAVE_OPERATIONS
 		{
 			return WaveReadLaneAt(x, GetWaveBroadcastSourceLaneIndex(Settings, WaveGetLaneIndex()));
 		}
@@ -159,7 +159,7 @@ uint WaveBroadcast(const FWaveBroadcastSettings Settings, uint x)
 		{
 			return WaveLaneRotateSwizzleRDNA(x, /* rotate_amount = */ Settings.Rotate, Settings.RotateFixMask);
 		}
-		#elif FEATURE_LEVEL >= FEATURE_LEVEL_SM6
+		#elif FEATURE_LEVEL >= FEATURE_LEVEL_SM6 || PLATFORM_SUPPORTS_SM6_0_WAVE_OPERATIONS
 		{
 			return WaveReadLaneAt(x, GetWaveBroadcastSourceLaneIndex(Settings, WaveGetLaneIndex()));
 		}
