Replaced fallback implementations of BitFieldExtractU32/BitFieldExtractI32 with more straightforward branchless versions

diff --git a/Engine/Shaders/Public/Platform.ush b/Engine/Shaders/Public/Platform.ush
--- a/Engine/Shaders/Public/Platform.ush
+++ b/Engine/Shaders/Public/Platform.ush
@@ -979,28 +979,20 @@ uint BitFieldInsertU32(uint Mask, uint Preserve, uint Enable)
 
 uint BitFieldExtractU32(uint Data, uint Size, uint Offset)
 {
-	Size &= 31u;
-	Offset &= 31u;
-
-	if (Size == 0u)
-		return 0u;
-	else if (Offset + Size < 32u)
-		return (Data << (32u - Size - Offset)) >> (32u - Size);
-	else
-		return Data >> Offset;
+	// Shift amounts are implicitly &31 in HLSL, so they should be optimized away on most platforms
+	// In GLSL shift amounts < 0 or >= word_size are undefined, so we better be explicit
+	Size &= 31;
+	Offset &= 31;
+	return (Data >> Offset) & ((1u << Size) - 1u);
 }
 
 int BitFieldExtractI32(int Data, uint Size, uint Offset)
 {
 	Size &= 31u;
 	Offset &= 31u;
-
-	if (Size == 0u)
-		return 0;
-	else if (Offset + Size < 32u)
-		return (Data << (32u - Size - Offset)) >> (32u - Size);
-	else
-		return Data >> Offset;
+	const uint Shift = (32u - Size) & 31u;
+	const int Value = (Data >> Offset) & int((1u << Size) - 1u);
+	return (Value << Shift) >> Shift;
 }
 
 uint BitFieldMaskU32(uint MaskWidth, uint MaskLocation)
