Only clear loop control mask instead of replacing Unroll with Loop mask when cross-compiling HLSL-to-HLSL.

#rnx

diff --git a/Engine/Shaders/Public/Platform/D3D/D3DCommon.ush b/Engine/Shaders/Public/Platform/D3D/D3DCommon.ush
--- a/Engine/Shaders/Public/Platform/D3D/D3DCommon.ush
+++ b/Engine/Shaders/Public/Platform/D3D/D3DCommon.ush
@@ -5,7 +5,7 @@ D3DCommon.usf: Common defines for the d3d shader compiler
 =============================================================================*/
 
 // Update this GUID to invalidate D3D shaders and force a shader recompilation for D3D RHIs.
-// GUID: B8706CF6-E7AA-43DB-91BC-60955A7B071B
+// GUID: 8E6C30C5-AD0A-4C9D-8AB3-CD2C92A4EFBC
 
 precise float MakePrecise(float v) { precise float pv = v; return pv; }
 precise float2 MakePrecise(float2 v) { precise float2 pv = v; return pv; }
diff --git a/Engine/Source/Developer/Windows/ShaderFormatD3D/Private/D3DShaderCompiler.cpp b/Engine/Source/Developer/Windows/ShaderFormatD3D/Private/D3DShaderCompiler.cpp
--- a/Engine/Source/Developer/Windows/ShaderFormatD3D/Private/D3DShaderCompiler.cpp
+++ b/Engine/Source/Developer/Windows/ShaderFormatD3D/Private/D3DShaderCompiler.cpp
@@ -583,7 +583,7 @@ static void PatchSpirvForPrecompilation(FSpirv& Spirv)
 	{
 		if (SpirvInstruction.Opcode() == SpvOpLoopMerge && SpirvInstruction.Operand(3) == SpvLoopControlUnrollMask)
 		{
-			(*SpirvInstruction)[3] = SpvLoopControlDontUnrollMask;
+			(*SpirvInstruction)[3] = SpvLoopControlMaskNone;
 		}
 	}
 }
