Bump GUID in VulkanCommon.ush to fix shader invalidation issue.

#rnx

diff --git a/Engine/Shaders/Public/Platform/Vulkan/VulkanCommon.ush b/Engine/Shaders/Public/Platform/Vulkan/VulkanCommon.ush
--- a/Engine/Shaders/Public/Platform/Vulkan/VulkanCommon.ush
+++ b/Engine/Shaders/Public/Platform/Vulkan/VulkanCommon.ush
@@ -5,7 +5,7 @@
 =============================================================================*/
 
 // Update this GUID to improve shader recompilation for Vulkan only shaders
-// GUID = 62c2037f-178e-46d1-9de1-f107cb1d1d39
+// GUID = 494E180E-0471-4013-8210-B52A0B05805F
 
 #pragma once
 
