Implements WAVESIZE()

diff --git a/Engine/Shaders/Private/TemporalSuperResolution/TSRRejectShading.usf b/Engine/Shaders/Private/TemporalSuperResolution/TSRRejectShading.usf
--- a/Engine/Shaders/Private/TemporalSuperResolution/TSRRejectShading.usf
+++ b/Engine/Shaders/Private/TemporalSuperResolution/TSRRejectShading.usf
@@ -451,6 +451,9 @@ tsr_simd_halfC ClampFireFliersWithGuide(tsr_simd_halfC ToClamp, tsr_simd_halfC G
 
 //------------------------------------------------------- ENTRY POINT
 
+#if COMPILER_SUPPORTS_WAVE_SIZE && DIM_WAVE_SIZE > 0
+	WAVESIZE(LANE_COUNT)
+#endif
 [numthreads(LANE_COUNT, 1, 1)]
 void MainCS(
 	uint2 GroupId : SV_GroupID,
diff --git a/Engine/Shaders/Public/Platform.ush b/Engine/Shaders/Public/Platform.ush
--- a/Engine/Shaders/Public/Platform.ush
+++ b/Engine/Shaders/Public/Platform.ush
@@ -93,6 +93,11 @@
 #define COMPILER_SUPPORTS_BARYCENTRIC_INTRINSICS 0
 #endif
 
+// Whether the shader compiler supports WAVESIZE()
+#ifndef COMPILER_SUPPORTS_WAVE_SIZE
+#define COMPILER_SUPPORTS_WAVE_SIZE 0
+#endif
+
 #ifndef COMPILER_FORCE_WAVE32_MODE
 #define COMPILER_FORCE_WAVE32_MODE
 #endif
diff --git a/Engine/Shaders/Public/Platform/D3D/D3DCommon.ush b/Engine/Shaders/Public/Platform/D3D/D3DCommon.ush
--- a/Engine/Shaders/Public/Platform/D3D/D3DCommon.ush
+++ b/Engine/Shaders/Public/Platform/D3D/D3DCommon.ush
@@ -7,6 +7,11 @@ D3DCommon.usf: Common defines for the d3d shader compiler
 // Update this GUID to invalidate D3D shaders and force a shader recompilation for D3D RHIs.
 // GUID: 8E6C30C5-AD0A-4C9D-8AB3-CD2C92A4EFBC
 
+#if SM6_PROFILE
+	#define COMPILER_SUPPORTS_WAVE_SIZE 1
+	#define WAVESIZE(N) [WaveSize(N)]
+#endif
+
 precise float MakePrecise(float v) { precise float pv = v; return pv; }
 precise float2 MakePrecise(float2 v) { precise float2 pv = v; return pv; }
 precise float3 MakePrecise(float3 v) { precise float3 pv = v; return pv; }
