Implemented support for 24bit integer operations, as well as some extra useful utilities like Padding, CeilLog2, and BitFieldExtractFloat

diff --git a/Engine/Shaders/Public/Platform.ush b/Engine/Shaders/Public/Platform.ush
--- a/Engine/Shaders/Public/Platform.ush
+++ b/Engine/Shaders/Public/Platform.ush
@@ -1034,6 +1034,50 @@ uint BitFieldMaskU32(uint MaskWidth, uint MaskLocation)
 }
 #endif
 
+#if !defined(COMPILER_SUPPORTS_24BIT_INTRINSICS)
+#define COMPILER_SUPPORTS_24BIT_INTRINSICS 0
+
+int MulI24(int I1, int I2)
+{
+	// D.i32 = S0.i24 * S1.i24
+	return (I1 * I2);
+}
+
+uint MulU24(uint U1, uint U2)
+{
+	// D.u32 = S0.u24 * S1.u24
+	return U1 * U2;
+}
+
+int MadI24(int I1, int I2, int I3)
+{
+	//  D.i = S0.i[23:0] * S1.i[23:0] + S2.i
+	return I1 * I2 + I3;
+}
+
+uint MadU24(uint U1, uint U2, uint U3)
+{
+	// D.u = S0.u[23:0] * S1.u[23:0] + S2.u
+	return U1 * U2 + U3;
+}
+
+#endif
+
+uint Padding(uint Value, uint Pow2)
+{
+	return (Value + Pow2 - 1u) & ~(Pow2 - 1u);
+}
+
+uint CeilLog2(uint Value)
+{
+	return Value < 2u ? 0u : firstbithigh(Value - 1u) + 1u;
+}
+
+float BitFieldExtractFloat(uint Bits, uint Count, uint Offset)
+{
+	return BitFieldExtractU32(Bits, Count, Offset) / (float)BitFieldMaskU32(Count, 0u);
+}
+
 #if !defined(COMPILER_SUPPORTS_BITALIGN)
 #define COMPILER_SUPPORTS_BITALIGN 0
 uint BitAlignU32(uint High, uint Low, uint Shift)
