[shaders] remove unnecessary literal to literal token concatenation in Platform.ush (STB preprocessor doesn't support this)

diff --git a/Engine/Shaders/Public/Platform.ush b/Engine/Shaders/Public/Platform.ush
--- a/Engine/Shaders/Public/Platform.ush
+++ b/Engine/Shaders/Public/Platform.ush
@@ -902,26 +902,26 @@ half4 max3( half4 a, half4 b, half4 c )
 	TYPE##2 select_internal(bool    c, TYPE    a, TYPE##2 b) { return TYPE##2(c   ? a   : b.x, c   ? a   : b.y); } \
 	TYPE##2 select_internal(bool    c, TYPE##2 a, TYPE    b) { return TYPE##2(c   ? a.x : b  , c   ? a.y : b  ); } \
 	TYPE##2 select_internal(bool    c, TYPE##2 a, TYPE##2 b) { return TYPE##2(c   ? a.x : b.x, c   ? a.y : b.y); } \
-	TYPE##2 select_internal(bool##2 c, TYPE    a, TYPE    b) { return TYPE##2(c.x ? a   : b  , c.y ? a   : b  ); } \
-	TYPE##2 select_internal(bool##2 c, TYPE    a, TYPE##2 b) { return TYPE##2(c.x ? a   : b.x, c.y ? a   : b.y); } \
-	TYPE##2 select_internal(bool##2 c, TYPE##2 a, TYPE    b) { return TYPE##2(c.x ? a.x : b  , c.y ? a.y : b  ); } \
-	TYPE##2 select_internal(bool##2 c, TYPE##2 a, TYPE##2 b) { return TYPE##2(c.x ? a.x : b.x, c.y ? a.y : b.y); } \
+	TYPE##2 select_internal(bool2   c, TYPE    a, TYPE    b) { return TYPE##2(c.x ? a   : b  , c.y ? a   : b  ); } \
+	TYPE##2 select_internal(bool2   c, TYPE    a, TYPE##2 b) { return TYPE##2(c.x ? a   : b.x, c.y ? a   : b.y); } \
+	TYPE##2 select_internal(bool2   c, TYPE##2 a, TYPE    b) { return TYPE##2(c.x ? a.x : b  , c.y ? a.y : b  ); } \
+	TYPE##2 select_internal(bool2   c, TYPE##2 a, TYPE##2 b) { return TYPE##2(c.x ? a.x : b.x, c.y ? a.y : b.y); } \
 	\
 	TYPE##3 select_internal(bool    c, TYPE    a, TYPE##3 b) { return TYPE##3(c   ? a   : b.x, c   ? a   : b.y, c   ? a   : b.z); } \
 	TYPE##3 select_internal(bool    c, TYPE##3 a, TYPE    b) { return TYPE##3(c   ? a.x : b  , c   ? a.y : b  , c   ? a.z : b  ); } \
 	TYPE##3 select_internal(bool    c, TYPE##3 a, TYPE##3 b) { return TYPE##3(c   ? a.x : b.x, c   ? a.y : b.y, c   ? a.z : b.z); } \
-	TYPE##3 select_internal(bool##3 c, TYPE    a, TYPE    b) { return TYPE##3(c.x ? a   : b  , c.y ? a   : b  , c.z ? a   : b  ); } \
-	TYPE##3 select_internal(bool##3 c, TYPE    a, TYPE##3 b) { return TYPE##3(c.x ? a   : b.x, c.y ? a   : b.y, c.z ? a   : b.z); } \
-	TYPE##3 select_internal(bool##3 c, TYPE##3 a, TYPE    b) { return TYPE##3(c.x ? a.x : b  , c.y ? a.y : b  , c.z ? a.z : b  ); } \
-	TYPE##3 select_internal(bool##3 c, TYPE##3 a, TYPE##3 b) { return TYPE##3(c.x ? a.x : b.x, c.y ? a.y : b.y, c.z ? a.z : b.z); } \
+	TYPE##3 select_internal(bool3   c, TYPE    a, TYPE    b) { return TYPE##3(c.x ? a   : b  , c.y ? a   : b  , c.z ? a   : b  ); } \
+	TYPE##3 select_internal(bool3   c, TYPE    a, TYPE##3 b) { return TYPE##3(c.x ? a   : b.x, c.y ? a   : b.y, c.z ? a   : b.z); } \
+	TYPE##3 select_internal(bool3   c, TYPE##3 a, TYPE    b) { return TYPE##3(c.x ? a.x : b  , c.y ? a.y : b  , c.z ? a.z : b  ); } \
+	TYPE##3 select_internal(bool3   c, TYPE##3 a, TYPE##3 b) { return TYPE##3(c.x ? a.x : b.x, c.y ? a.y : b.y, c.z ? a.z : b.z); } \
 	\
 	TYPE##4 select_internal(bool    c, TYPE    a, TYPE##4 b) { return TYPE##4(c   ? a   : b.x, c   ? a   : b.y, c   ? a   : b.z, c   ? a   : b.w); } \
 	TYPE##4 select_internal(bool    c, TYPE##4 a, TYPE    b) { return TYPE##4(c   ? a.x : b  , c   ? a.y : b  , c   ? a.z : b  , c   ? a.w : b  ); } \
 	TYPE##4 select_internal(bool    c, TYPE##4 a, TYPE##4 b) { return TYPE##4(c   ? a.x : b.x, c   ? a.y : b.y, c   ? a.z : b.z, c   ? a.w : b.w); } \
-	TYPE##4 select_internal(bool##4 c, TYPE    a, TYPE    b) { return TYPE##4(c.x ? a   : b  , c.y ? a   : b  , c.z ? a   : b  , c.w ? a   : b  ); } \
-	TYPE##4 select_internal(bool##4 c, TYPE    a, TYPE##4 b) { return TYPE##4(c.x ? a   : b.x, c.y ? a   : b.y, c.z ? a   : b.z, c.w ? a   : b.w); } \
-	TYPE##4 select_internal(bool##4 c, TYPE##4 a, TYPE    b) { return TYPE##4(c.x ? a.x : b  , c.y ? a.y : b  , c.z ? a.z : b  , c.w ? a.w : b  ); } \
-	TYPE##4 select_internal(bool##4 c, TYPE##4 a, TYPE##4 b) { return TYPE##4(c.x ? a.x : b.x, c.y ? a.y : b.y, c.z ? a.z : b.z, c.w ? a.w : b.w); } \
+	TYPE##4 select_internal(bool4   c, TYPE    a, TYPE    b) { return TYPE##4(c.x ? a   : b  , c.y ? a   : b  , c.z ? a   : b  , c.w ? a   : b  ); } \
+	TYPE##4 select_internal(bool4   c, TYPE    a, TYPE##4 b) { return TYPE##4(c.x ? a   : b.x, c.y ? a   : b.y, c.z ? a   : b.z, c.w ? a   : b.w); } \
+	TYPE##4 select_internal(bool4   c, TYPE##4 a, TYPE    b) { return TYPE##4(c.x ? a.x : b  , c.y ? a.y : b  , c.z ? a.z : b  , c.w ? a.w : b  ); } \
+	TYPE##4 select_internal(bool4   c, TYPE##4 a, TYPE##4 b) { return TYPE##4(c.x ? a.x : b.x, c.y ? a.y : b.y, c.z ? a.z : b.z, c.w ? a.w : b.w); } \
 
 DEFINE_SELECT(bool)
 DEFINE_SELECT(uint)
