Added bSupportsClipDistance to DDSPI for mobile shader platforms that support SV_ClipDistance

diff --git a/Engine/Shaders/Public/Platform.ush b/Engine/Shaders/Public/Platform.ush
--- a/Engine/Shaders/Public/Platform.ush
+++ b/Engine/Shaders/Public/Platform.ush
@@ -358,7 +358,6 @@
 
 // Type macros for Uniform Buffer code generation.
 // Hides the platform specific 'half' and 'fixed' support but has to extend to all types and dimensions
-#if (FEATURE_LEVEL == FEATURE_LEVEL_ES3_1)
 	#define UB_INT(Dim)         int##Dim
 	#define UB_UINT(Dim)        uint##Dim
 	#define UB_FLOAT(Dim)       float##Dim
