Remove unsigned integer instance of overloaded HLSL function "select_internal" to fix ambiguity.

#rnx

diff --git a/Engine/Shaders/Public/Platform.ush b/Engine/Shaders/Public/Platform.ush
--- a/Engine/Shaders/Public/Platform.ush
+++ b/Engine/Shaders/Public/Platform.ush
@@ -924,12 +924,10 @@ half4 max3( half4 a, half4 b, half4 c )
 	TYPE##4 select_internal(bool4   c, TYPE##4 a, TYPE##4 b) { return TYPE##4(c.x ? a.x : b.x, c.y ? a.y : b.y, c.z ? a.z : b.z, c.w ? a.w : b.w); } \
 
 DEFINE_SELECT(bool)
-DEFINE_SELECT(uint)
 DEFINE_SELECT(int)
 DEFINE_SELECT(float)
 #if PLATFORM_SUPPORTS_REAL_TYPES
 DEFINE_SELECT(half)
-DEFINE_SELECT(uint16_t)
 DEFINE_SELECT(int16_t)
 #endif
 #undef DEFINE_SELECT
