Limit removal of unsigned integer instance of overloaded HLSL function "select_internal" to COMPILER_DXC and COMPILER_HLSLCC only. Appendix to CL 23511340.

#rnx

diff --git a/Engine/Shaders/Public/Platform.ush b/Engine/Shaders/Public/Platform.ush
--- a/Engine/Shaders/Public/Platform.ush
+++ b/Engine/Shaders/Public/Platform.ush
@@ -924,10 +924,18 @@ half4 max3( half4 a, half4 b, half4 c )
 	TYPE##4 select_internal(bool4   c, TYPE##4 a, TYPE##4 b) { return TYPE##4(c.x ? a.x : b.x, c.y ? a.y : b.y, c.z ? a.z : b.z, c.w ? a.w : b.w); } \
 
 DEFINE_SELECT(bool)
+#if !(COMPILER_DXC || COMPILER_HLSLCC)
+// @todo-lh: Ambiguous for DXC and HLSLcc if no suffix is provided for integer literals
+DEFINE_SELECT(uint)
+#endif
 DEFINE_SELECT(int)
 DEFINE_SELECT(float)
 #if PLATFORM_SUPPORTS_REAL_TYPES
 DEFINE_SELECT(half)
+#if !(COMPILER_DXC || COMPILER_HLSLCC)
+// @todo-lh: Ambiguous for DXC and HLSLcc if no suffix is provided for integer literals
+DEFINE_SELECT(uint16_t)
+#endif
 DEFINE_SELECT(int16_t)
 #endif
 #undef DEFINE_SELECT
