Change UEReportAssertWithPayload to use an inifinite loop to trigger a timeout instead of attempting to trigger a page fault through OOB buffer access

- This makes the assert reporting work correctly on AMD as well as NVIDIA GPUs

diff --git a/Engine/Shaders/Public/Platform/D3D/D3DCommon.ush b/Engine/Shaders/Public/Platform/D3D/D3DCommon.ush
--- a/Engine/Shaders/Public/Platform/D3D/D3DCommon.ush
+++ b/Engine/Shaders/Public/Platform/D3D/D3DCommon.ush
@@ -55,9 +55,16 @@ void UEReportAssertWithPayload(uint AssertID, uint4 Payload)
 			UEDiagnosticBuffer[4] = Payload.z;
 			UEDiagnosticBuffer[5] = Payload.w;
 		}
-		// Try to trigger a crash
+		// Try to trigger a crash (timeout) using a deliberate infinite loop
 		// TODO: use PLATFORM_BREAK
-		UEDiagnosticBuffer[0xFFFFFFFF] = 0;
+		for (;;)
+		{
+			InterlockedAdd(UEDiagnosticBuffer[0], 0, OldValue);
+			if (OldValue == 0)
+			{
+				break;
+			}
+		}
 	}
 }
 
