Invalidating shaders after FVirtualVoxelParameters was changed (via FHairStrandsVoxelCommonParameters).

diff --git a/Engine/Shaders/Public/ShaderVersion.ush b/Engine/Shaders/Public/ShaderVersion.ush
--- a/Engine/Shaders/Public/ShaderVersion.ush
+++ b/Engine/Shaders/Public/ShaderVersion.ush
@@ -2,4 +2,4 @@
 // Each time the console command is executed it generates a new GUID. As this file is included
 // in Platform.ush (which should be included in any shader) it allows to invalidate the shader DDC.
 // 
-// GUID = A9EC3819-3FCC-41DC-88FB-A1AC66A46D20
+// GUID = DB03DB78-EB6A-46FD-8894-A1BB690894F9
