Fix for non-strata cook.

diff --git a/Engine/Shaders/Public/StrataPublic.ush b/Engine/Shaders/Public/StrataPublic.ush
--- a/Engine/Shaders/Public/StrataPublic.ush
+++ b/Engine/Shaders/Public/StrataPublic.ush
@@ -8,6 +8,8 @@
 
 #include "/Engine/Private/Strata/Strata.ush"
 
+#if STRATA_ENABLED
+
 /** Public facing functions allowing external users access to some Strata data. */
 float3 StrataPublic_GetWorldNormal(uint2 PixelCoord)
 {
@@ -15,3 +17,5 @@ float3 StrataPublic_GetWorldNormal(uint2 PixelCoord)
 }
 
 //TODO: Implement other accessors that were previously available via SceneTextures.
+
+#endif
\ No newline at end of file
