Implements PLATFORM_GPU_ARCH

diff --git a/Engine/Shaders/Public/Platform.ush b/Engine/Shaders/Public/Platform.ush
--- a/Engine/Shaders/Public/Platform.ush
+++ b/Engine/Shaders/Public/Platform.ush
@@ -17,6 +17,35 @@
 #define FEATURE_LEVEL_SM6			6
 #define FEATURE_LEVEL_MAX			7
 
+// Values of PLATFORM_GPU_ARCH
+#define PLATFORM_GPU_ARCH_UNKNOWN          0x0000
+
+#define PLATFORM_GPU_ARCH_AMD              0x1000
+#define PLATFORM_GPU_ARCH_AMD_GCN_1        0x1010
+#define PLATFORM_GPU_ARCH_AMD_GCN_2        0x1020
+#define PLATFORM_GPU_ARCH_AMD_GCN_3        0x1030
+#define PLATFORM_GPU_ARCH_AMD_GCN_4        0x1040
+#define PLATFORM_GPU_ARCH_AMD_GCN_5        0x1050
+#define PLATFORM_GPU_ARCH_AMD_RDNA_1       0x1110
+#define PLATFORM_GPU_ARCH_AMD_RDNA_2       0x1120
+#define PLATFORM_GPU_ARCH_AMD_RDNA_3       0x1130
+#define PLATFORM_GPU_ARCH_AMD_LATTEST      PLATFORM_GPU_ARCH_AMD_RDNA_3
+
+#define PLATFORM_GPU_ARCH_NVIDIA           0x2000
+#define PLATFORM_GPU_ARCH_NVIDIA_KEPLER    0x2010
+#define PLATFORM_GPU_ARCH_NVIDIA_MAXWELL   0x2020
+#define PLATFORM_GPU_ARCH_NVIDIA_PASCAL    0x2030
+#define PLATFORM_GPU_ARCH_NVIDIA_VOLTA     0x2040
+#define PLATFORM_GPU_ARCH_NVIDIA_TURING    0x2050
+#define PLATFORM_GPU_ARCH_NVIDIA_AMPERE    0x2060
+#define PLATFORM_GPU_ARCH_NVIDIA_ADA       0x2070
+#define PLATFORM_GPU_ARCH_NVIDIA_LATTEST   PLATFORM_GPU_ARCH_NVIDIA_ADA
+
+#define PLATFORM_GPU_ARCH_INTEL            0x3000
+#define PLATFORM_GPU_ARCH_INTEL_ARC        0x3010
+#define PLATFORM_GPU_ARCH_INTEL_LATTEST    PLATFORM_GPU_ARCH_INTEL_ARC
+
+
 // ---------------------------------------------------- Profile or compiler specific includes
 // TODO: Have shader compiler including these platform specific USF files, that needs to work
 // with ShaderCore.cpp's GetShaderIncludes().
@@ -238,6 +267,10 @@ float2 UnpackFloat2FromUInt(uint In)
 #define PLATFORM_SUPPORTS_UB_STRUCT 0
 #endif
 
+#ifndef PLATFORM_GPU_ARCH
+#define PLATFORM_GPU_ARCH PLATFORM_GPU_ARCH_UNKNOWN
+#endif 
+
 #if PLATFORM_SUPPORTS_UB_STRUCT
 	#define UB_CB_NAME(X) UniformBufferConstants_##X
 #else
