Bumped shader version on Mac to work around serialization bug.

diff --git a/Engine/Shaders/Public/Platform/Metal/MetalCommon.ush b/Engine/Shaders/Public/Platform/Metal/MetalCommon.ush
--- a/Engine/Shaders/Public/Platform/Metal/MetalCommon.ush
+++ b/Engine/Shaders/Public/Platform/Metal/MetalCommon.ush
@@ -5,7 +5,7 @@
 =============================================================================*/
 
 // Update this GUID to improve shader recompilation for Metal only shaders
-// GUID = C53FCB9D-55EE-439A-BD5C-EFB947FEF98E
+// GUID = 9984851C-6B85-4143-A6F9-F81D5708EFB6
 
 #pragma once
 
