Bumping the metal shader GUID because of the number of shader crashes due to out-of-date shaders.

diff --git a/Engine/Shaders/Public/Platform/Metal/MetalCommon.ush b/Engine/Shaders/Public/Platform/Metal/MetalCommon.ush
--- a/Engine/Shaders/Public/Platform/Metal/MetalCommon.ush
+++ b/Engine/Shaders/Public/Platform/Metal/MetalCommon.ush
@@ -5,7 +5,7 @@
 =============================================================================*/
 
 // Update this GUID to improve shader recompilation for Metal only shaders
-// GUID = C7874545-4A87-459D-A554-ED1BB12BE194
+// GUID = DF88BF25-6438-4D0F-9236-52E326D5C9B3
 
 #pragma once
 
