Bump shader version to clean up suspected DDC corruption from shader compiler race condition.

#rnx

diff --git a/Engine/Shaders/Public/ShaderVersion.ush b/Engine/Shaders/Public/ShaderVersion.ush
--- a/Engine/Shaders/Public/ShaderVersion.ush
+++ b/Engine/Shaders/Public/ShaderVersion.ush
@@ -2,4 +2,4 @@
 // Each time the console command is executed it generates a new GUID. As this file is included
 // in Platform.ush (which should be included in any shader) it allows to invalidate the shader DDC.
 // 
-// GUID = 9CAC9961-9C56-48D8-947A-8A6A2D736E60
+// GUID = 2AE92E40-F312-4448-9627-C1AD6B6176A5
\ No newline at end of file
