[shaders] (hopefully final) version bump to clear out QAGame bad ddc data

#rnx

diff --git a/Engine/Shaders/Public/ShaderVersion.ush b/Engine/Shaders/Public/ShaderVersion.ush
--- a/Engine/Shaders/Public/ShaderVersion.ush
+++ b/Engine/Shaders/Public/ShaderVersion.ush
@@ -2,4 +2,4 @@
 // Each time the console command is executed it generates a new GUID. As this file is included
 // in Platform.ush (which should be included in any shader) it allows to invalidate the shader DDC.
 // 
-// GUID = 2AE92E40-F312-4448-9627-C1AD6B6176A5
\ No newline at end of file
+// GUID = 2408A7EF-2137-4013-BE40-8D684E9876E1
\ No newline at end of file
