VulkanRHI: VKRT WIP: Add SystemParameters as shader_record to all ray tracing shaders.

#rnx

diff --git a/Engine/Shaders/Public/Platform/Vulkan/VulkanCommon.ush b/Engine/Shaders/Public/Platform/Vulkan/VulkanCommon.ush
--- a/Engine/Shaders/Public/Platform/Vulkan/VulkanCommon.ush
+++ b/Engine/Shaders/Public/Platform/Vulkan/VulkanCommon.ush
@@ -5,7 +5,7 @@
 =============================================================================*/
 
 // Update this GUID to improve shader recompilation for Vulkan only shaders
-// GUID = EEA3F6F8-2405-4093-9A65-5072C38494AC
+// GUID = 4BB75D9E-9229-4BB3-959A-36A81475A4FE
 
 #pragma once
 
@@ -63,3 +63,30 @@ void ImageInterlockedMaxUInt64(RWTexture2D<UlongType> Dest, uint2 Coord, UlongTy
 	COMPILER_DEFINE VULKAN_WAVESIZE(N)
 	#define WAVESIZE(N) VULKAN_WAVESIZE(N)
 #endif
+
+#if PLATFORM_SUPPORTS_BINDLESS && (RAYCALLABLESHADER || RAYHITGROUPSHADER || RAYGENSHADER || RAYMISSSHADER)
+	#define OVERRIDE_RAY_TRACING_HIT_GROUP_SYSTEM_RESOURCES 1
+	#include "/Engine/Shared/RayTracingBuiltInResources.h"
+
+	// Combine the common FHitGroupSystemRootConstants with bindless parameters
+	// :todo-jn: must match with VulkanRayTracing.h decl until we create a common header like RayTracingBuiltInResources.h
+	struct FVulkanHitGroupSystemParameters
+	{
+		FHitGroupSystemRootConstants RootConstants;
+
+		uint BindlessHitGroupSystemIndexBuffer;
+		uint BindlessHitGroupSystemVertexBuffer;
+
+		uint BindlessUniformBuffers[VULKAN_MAX_BINDLESS_UNIFORM_BUFFERS_PER_STAGE];
+	};
+
+	[[vk::shader_record_ext]] ConstantBuffer<FVulkanHitGroupSystemParameters> VulkanHitGroupSystemParameters;
+
+	static const FHitGroupSystemRootConstants HitGroupSystemRootConstants = VulkanHitGroupSystemParameters.RootConstants;
+
+	static const uint BindlessResource_HitGroupSystemIndexBuffer = VulkanHitGroupSystemParameters.BindlessHitGroupSystemIndexBuffer;
+	UB_RESOURCE_MEMBER_RESOURCE(ByteAddressBuffer, HitGroupSystemIndexBuffer);
+	
+	static const uint BindlessResource_HitGroupSystemVertexBuffer = VulkanHitGroupSystemParameters.BindlessHitGroupSystemVertexBuffer;
+	UB_RESOURCE_MEMBER_RESOURCE(ByteAddressBuffer, HitGroupSystemVertexBuffer);
+#endif
diff --git a/Engine/Shaders/Shared/RayTracingBuiltInResources.h b/Engine/Shaders/Shared/RayTracingBuiltInResources.h
--- a/Engine/Shaders/Shared/RayTracingBuiltInResources.h
+++ b/Engine/Shaders/Shared/RayTracingBuiltInResources.h
@@ -16,7 +16,7 @@
 #if defined(__cplusplus)
 	#define INCLUDED_FROM_CPP_CODE  1
 	#define INCLUDED_FROM_HLSL_CODE 0
-#elif defined(SM5_PROFILE)
+#elif defined(SM5_PROFILE) || defined(VULKAN_PROFILE_SM6)
 	// #dxr_todo: we should use a built-in macro to detect if this shader is compiled using DXC (depends on [REDACTED_PROFILE_URL]/DirectXShaderCompiler/issues/1686)
 	#define INCLUDED_FROM_CPP_CODE  0
 	#define INCLUDED_FROM_HLSL_CODE 1
diff --git a/Engine/Source/Developer/VulkanShaderFormat/Private/VulkanShaderCompiler.cpp b/Engine/Source/Developer/VulkanShaderFormat/Private/VulkanShaderCompiler.cpp
--- a/Engine/Source/Developer/VulkanShaderFormat/Private/VulkanShaderCompiler.cpp
+++ b/Engine/Source/Developer/VulkanShaderFormat/Private/VulkanShaderCompiler.cpp
@@ -2645,6 +2645,7 @@ void DoCompileVulkanShader(const FShaderCompilerInput& Input, FShaderCompilerOut
 
 	AdditionalDefines.SetDefine(TEXT("VULKAN_BINDLESS_SAMPLER_ARRAY_PREFIX"), VulkanBindless::kBindlessSamplerArrayPrefix);
 	AdditionalDefines.SetDefine(TEXT("VULKAN_BINDLESS_RESOURCE_ARRAY_PREFIX"), VulkanBindless::kBindlessResourceArrayPrefix);
+	AdditionalDefines.SetDefine(TEXT("VULKAN_MAX_BINDLESS_UNIFORM_BUFFERS_PER_STAGE"), VulkanBindless::MaxUniformBuffersPerStage);
 
 	if (IsAndroidShaderFormat(Input.ShaderFormat))
 	{
diff --git a/Engine/Source/Runtime/VulkanRHI/Public/VulkanCommon.h b/Engine/Source/Runtime/VulkanRHI/Public/VulkanCommon.h
--- a/Engine/Source/Runtime/VulkanRHI/Public/VulkanCommon.h
+++ b/Engine/Source/Runtime/VulkanRHI/Public/VulkanCommon.h
@@ -114,7 +114,7 @@ namespace ShaderStage
 
 namespace VulkanBindless
 {
-	static constexpr uint32 MaxUniformBuffersPerStage = 8;
+	static constexpr uint32 MaxUniformBuffersPerStage = 16;
 
 	enum EDescriptorSets
 	{
