Improves TSR's signal to noise ratio on high noise surfaces under movement

diff --git a/Engine/Shaders/Private/TemporalSuperResolution/TSRConvolutionNetwork.ush b/Engine/Shaders/Private/TemporalSuperResolution/TSRConvolutionNetwork.ush
--- a/Engine/Shaders/Private/TemporalSuperResolution/TSRConvolutionNetwork.ush
+++ b/Engine/Shaders/Private/TemporalSuperResolution/TSRConvolutionNetwork.ush
@@ -289,7 +289,7 @@ void AccessNeighborTexels1x3(
 
 CALL_SITE_DEBUGLOC
 template<typename ScalarType, uint VectorSize, uint LaneStrideX, uint LaneStrideY>
-TLaneVector2D<ScalarType, VectorSize, LaneStrideX, LaneStrideY> Blur3x3(TLaneVector2D<ScalarType, VectorSize, LaneStrideX, LaneStrideY> Center)
+TLaneVector2D<ScalarType, VectorSize, LaneStrideX, LaneStrideY> WeightedSum3x3(TLaneVector2D<ScalarType, VectorSize, LaneStrideX, LaneStrideY> Center, float3 HorizontalWeights, float3 VerticalWeights)
 {
 	TLaneVector2D<ScalarType, VectorSize, LaneStrideX, LaneStrideY> CO = Center;
 	TLaneVector2D<ScalarType, VectorSize, LaneStrideX, LaneStrideY> CP;
@@ -297,20 +297,68 @@ TLaneVector2D<ScalarType, VectorSize, LaneStrideX, LaneStrideY> Blur3x3(TLaneVec
 	AccessNeighborTexels1x3(CO, /* out */ CP, /* out */ CN);
 	
 	TLaneVector2D<ScalarType, VectorSize, LaneStrideX, LaneStrideY> BO = (
-		CO * TLaneVector2D<ScalarType, VectorSize, LaneStrideX, LaneStrideY>::Const(0.5) +
-		CP * TLaneVector2D<ScalarType, VectorSize, LaneStrideX, LaneStrideY>::Const(0.25) +
-		CN * TLaneVector2D<ScalarType, VectorSize, LaneStrideX, LaneStrideY>::Const(0.25));
+		CO * TLaneVector2D<ScalarType, VectorSize, LaneStrideX, LaneStrideY>::Const(ScalarType(VerticalWeights.y)) +
+		CP * TLaneVector2D<ScalarType, VectorSize, LaneStrideX, LaneStrideY>::Const(ScalarType(VerticalWeights.z)) +
+		CN * TLaneVector2D<ScalarType, VectorSize, LaneStrideX, LaneStrideY>::Const(ScalarType(VerticalWeights.x)));
 
 	TLaneVector2D<ScalarType, VectorSize, LaneStrideX, LaneStrideY> BP;
 	TLaneVector2D<ScalarType, VectorSize, LaneStrideX, LaneStrideY> BN;
 	AccessNeighborTexels3x1(BO, /* out */ BP, /* out */ BN);
 
 	return (
-		BO * TLaneVector2D<ScalarType, VectorSize, LaneStrideX, LaneStrideY>::Const(0.5) +
-		BP * TLaneVector2D<ScalarType, VectorSize, LaneStrideX, LaneStrideY>::Const(0.25) +
-		BN * TLaneVector2D<ScalarType, VectorSize, LaneStrideX, LaneStrideY>::Const(0.25));
+		BO * TLaneVector2D<ScalarType, VectorSize, LaneStrideX, LaneStrideY>::Const(ScalarType(HorizontalWeights.y)) +
+		BP * TLaneVector2D<ScalarType, VectorSize, LaneStrideX, LaneStrideY>::Const(ScalarType(HorizontalWeights.z)) +
+		BN * TLaneVector2D<ScalarType, VectorSize, LaneStrideX, LaneStrideY>::Const(ScalarType(HorizontalWeights.x)));
 }
 
+CALL_SITE_DEBUGLOC
+template<typename ScalarType, uint VectorSize, uint LaneStrideX, uint LaneStrideY>
+TLaneVector2D<ScalarType, VectorSize, LaneStrideX, LaneStrideY> Sum3x3(TLaneVector2D<ScalarType, VectorSize, LaneStrideX, LaneStrideY> Center)
+{
+	return WeightedSum3x3(Center, /* HorizontalWeights = */ float(1.0).xxx, /* VerticalWeights = */ float(1.0).xxx);
+}
+
+CALL_SITE_DEBUGLOC
+template<typename ScalarType, uint VectorSize, uint LaneStrideX, uint LaneStrideY>
+TLaneVector2D<ScalarType, VectorSize, LaneStrideX, LaneStrideY> WeightedAvg3x3(TLaneVector2D<ScalarType, VectorSize, LaneStrideX, LaneStrideY> Center, float3 NormalizedWeights)
+{
+	return WeightedSum3x3(Center, /* HorizontalWeights = */ NormalizedWeights, /* VerticalWeights = */ NormalizedWeights);
+}
+
+CALL_SITE_DEBUGLOC
+template<typename ScalarType, uint VectorSize, uint LaneStrideX, uint LaneStrideY>
+TLaneVector2D<ScalarType, VectorSize, LaneStrideX, LaneStrideY> Blur3x3(TLaneVector2D<ScalarType, VectorSize, LaneStrideX, LaneStrideY> Center)
+{
+	return WeightedAvg3x3(Center, /* NormalizedWeights = */ float3(0.25, 0.5, 0.25));
+}
+
+CALL_SITE_DEBUGLOC
+template<typename ScalarType, uint VectorSize, uint LaneStrideX, uint LaneStrideY>
+TLaneVector2D<ScalarType, VectorSize, LaneStrideX, LaneStrideY> Sharpen3x3(TLaneVector2D<ScalarType, VectorSize, LaneStrideX, LaneStrideY> Center)
+{
+	return WeightedAvg3x3(Center, /* NormalizedWeights = */ float3(-0.25, 1.5, -0.25));
+}
+
+
+//------------------------------------------------------- TOTAL VARIATION 3x3 CONVOLUTION
+
+CALL_SITE_DEBUGLOC
+template<typename ScalarType, uint VectorSize, uint LaneStrideX, uint LaneStrideY>
+TLaneVector2D<ScalarType, VectorSize, LaneStrideX, LaneStrideY> TotalVariation3x3(TLaneVector2D<ScalarType, VectorSize, LaneStrideX, LaneStrideY> Center)
+#if 0
+{
+	TLaneVector2D<ScalarType, VectorSize, LaneStrideX, LaneStrideY> Sum = Blur3x3(Center);
+	
+	return TLaneVector2D<ScalarType, VectorSize, LaneStrideX, LaneStrideY>::Const(1) * Center - Sum;
+}
+#else
+{
+	TLaneVector2D<ScalarType, VectorSize, LaneStrideX, LaneStrideY> Sum = WeightedSum3x3(Center, /* HorizontalWeights = */ float(1.0).xxx, /* VerticalWeights = */ float3(0.125, 0.125, 0.125));
+
+	return TLaneVector2D<ScalarType, VectorSize, LaneStrideX, LaneStrideY>::Const(1.125) * Center - Sum;
+}
+#endif
+
 
 //------------------------------------------------------- MIN & MAX 3x3 CONVOLUTION
 
@@ -381,6 +429,42 @@ void MinMax3x3(
 }
 #endif
 
+CALL_SITE_DEBUGLOC
+template<typename ScalarType, uint VectorSize, uint LaneStrideX, uint LaneStrideY>
+TLaneVector2D<ScalarType, VectorSize, LaneStrideX, LaneStrideY> MaxMinusMin3x3(
+	TLaneVector2D<ScalarType, VectorSize, LaneStrideX, LaneStrideY> Center)
+{
+	TLaneVector2D<ScalarType, VectorSize, LaneStrideX, LaneStrideY> BoxMin;
+	TLaneVector2D<ScalarType, VectorSize, LaneStrideX, LaneStrideY> BoxMax;
+	MinMax3x3(Center, /* out */ BoxMin, /* out */ BoxMax);
+	return BoxMax - BoxMin;
+}
+
+
+//------------------------------------------------------- CLAMP CONVOLUTIONS
+
+CALL_SITE_DEBUGLOC
+template<typename ScalarType, uint VectorSize, uint LaneStrideX, uint LaneStrideY>
+TLaneVector2D<ScalarType, VectorSize, LaneStrideX, LaneStrideY> Clamp3x3(
+	TLaneVector2D<ScalarType, VectorSize, LaneStrideX, LaneStrideY> ToClamp,
+	TLaneVector2D<ScalarType, VectorSize, LaneStrideX, LaneStrideY> BoundaryCenter)
+{
+	TLaneVector2D<ScalarType, VectorSize, LaneStrideX, LaneStrideY> BoundaryMin;
+	TLaneVector2D<ScalarType, VectorSize, LaneStrideX, LaneStrideY> BoundaryMax;
+	MinMax3x3(BoundaryCenter, /* out */ BoundaryMin, /* out */ BoundaryMax);
+
+	return fastClamp(ToClamp, BoundaryMin, BoundaryMax);
+}
+
+CALL_SITE_DEBUGLOC
+template<typename ScalarType, uint VectorSize, uint LaneStrideX, uint LaneStrideY>
+TLaneVector2D<ScalarType, VectorSize, LaneStrideX, LaneStrideY> AnnihilateToGuide3x3(
+	TLaneVector2D<ScalarType, VectorSize, LaneStrideX, LaneStrideY> ToClamp,
+	TLaneVector2D<ScalarType, VectorSize, LaneStrideX, LaneStrideY> Guide)
+{
+	return Clamp3x3(ToClamp, Clamp3x3(Guide, ToClamp));
+}
+
 
 //------------------------------------------------------- MEDIAN 3x3 CONVOLUTION
 
diff --git a/Engine/Shaders/Private/TemporalSuperResolution/TSRRejectShading.usf b/Engine/Shaders/Private/TemporalSuperResolution/TSRRejectShading.usf
--- a/Engine/Shaders/Private/TemporalSuperResolution/TSRRejectShading.usf
+++ b/Engine/Shaders/Private/TemporalSuperResolution/TSRRejectShading.usf
@@ -17,9 +17,6 @@
 
 #define DEBUG_ARRAY_SIZE 8
 
-/** Whether should remove firefliers from input and history before comparison to improve stability. */
-#define CONFIG_FIREFLY_REMOVAL 1
-
 /** Whether to compose translucency to the input. */
 #define CONFIG_COMPOSE_TRANSLUCENCY 1
 
@@ -131,22 +128,6 @@ RWTexture2D<tsr_half> InputSceneColorLdrLumaOutput;
 
 //------------------------------------------------------- GENERAL OPERATIONS
 
-CALL_SITE_DEBUGLOC
-tsr_simd_halfC ClampFireFliersWithGuide(tsr_simd_halfC ToClamp, tsr_simd_halfC Guide)
-{
-	tsr_simd_halfC ToClampMin;
-	tsr_simd_halfC ToClampMax;
-	MinMax3x3(ToClamp, /* out */ ToClampMin, /* out */ ToClampMax);
-
-	tsr_simd_halfC ClampedGuide = fastClamp(Guide, ToClampMin, ToClampMax);
-	
-	tsr_simd_halfC ClampedGuideMin;
-	tsr_simd_halfC ClampedGuideMax;
-	MinMax3x3(ClampedGuide, /* out */ ClampedGuideMin, /* out */ ClampedGuideMax);
-	
-	return fastClamp(ToClamp, ClampedGuideMin, ClampedGuideMax);
-}
-
 tsr_short2 ComputeInputPixelPos(const tsr_short2 GroupPixelOffset, const uint LaneIndex, const uint SimdIndex)
 {
 	return GroupPixelOffset + GetLaneSimdPixelOffset<LANE_STRIDE_X, LANE_STRIDE_Y>(LaneIndex, SimdIndex);
@@ -489,28 +470,22 @@ void MainCS(
 			ExposedHistory = ExposedHistory * rcp(ExposedHistory + tsr_simd_halfC::Const(tsr_half(PerceptionAdd)));
 		}
 
-		// Top secret sauce with paper and salt.
-		tsr_simd_halfC ClampedInput = ExposedInput;
-		tsr_simd_halfC ClampedHistory = ExposedHistory;
-		#if CONFIG_FIREFLY_REMOVAL
-		{
-			ClampedInput   = ClampFireFliersWithGuide(ClampedInput  , /* Guide = */ ExposedHistory);
-			ClampedHistory = ClampFireFliersWithGuide(ClampedHistory, /* Guide = */ ExposedInput);
-		}
-		#endif
-
+		// Mutual Annihilation
+		tsr_simd_halfC ClampedInput   = AnnihilateToGuide3x3(ExposedInput  , /* Guide = */ ExposedHistory);
+		tsr_simd_halfC ClampedHistory = AnnihilateToGuide3x3(ExposedHistory, /* Guide = */ ExposedInput);
+		
 		// Filter clamped input and history to increase quantization of the rejection detection
 		tsr_simd_halfC FilteredInput   = Blur3x3(ClampedInput);
 		tsr_simd_halfC FilteredHistory = Blur3x3(ClampedHistory);
 	
+		// Build the clamping box of the filtered signal
 		tsr_simd_halfC FilteredBoxMin;
 		tsr_simd_halfC FilteredBoxMax;
 		MinMax3x3(FilteredInput, /* out */ FilteredBoxMin, /* out */ FilteredBoxMax);
 		FilteredBoxMin = FilteredBoxMin - BackbufferQuantizationErrorVector;
 		FilteredBoxMax = FilteredBoxMax + BackbufferQuantizationErrorVector;
 		
-		tsr_simd_halfC FlickeringClampedFilteredHistory = fastClamp(FilteredHistory, FilteredBoxMin, FilteredBoxMax);
-
+		// Clamp filtered signal
 		tsr_simd_halfC ClampedFilteredHistory;
 		#if CONFIG_MOIRE_DETECTION
 		{
@@ -538,35 +513,33 @@ void MainCS(
 			ClampedFilteredHistory = fastClamp(FilteredHistory, StableFilteredBoxMin, StableFilteredBoxMax);
 		}
 		#else
-			ClampedFilteredHistory = FlickeringClampedFilteredHistory;
+			ClampedFilteredHistory = fastClamp(FilteredHistory, FilteredBoxMin, FilteredBoxMax);
 		#endif
 		
-		tsr_simd_halfC OriginalBoxMin;
-		tsr_simd_halfC OriginalBoxMax;
-		MinMax3x3(ExposedInput, /* out */ OriginalBoxMin, /* out */ OriginalBoxMax);
-		tsr_simd_halfC BoxSize = (OriginalBoxMax - OriginalBoxMin + (BackbufferQuantizationErrorVector + BackbufferQuantizationErrorVector)) * FilteringWeight;
-		
-		tsr_simd_halfC RawClampedEnergy = abs(ClampedFilteredHistory - FilteredHistory);
-
-		tsr_simd_halfC FilteredClampedEnergy;
+		// Compute the box size.
+		tsr_simd_halfC BoxSize;
 		{
-			FilteredClampedEnergy = Median3x3(RawClampedEnergy);
-			FilteredClampedEnergy = Max3x3(FilteredClampedEnergy);
-		}
+			BoxSize = (
+				MaxMinusMin3x3(ClampedInput) * tsr_simd_halfC::Const(FilteringWeight) +
+				(BackbufferQuantizationErrorVector * tsr_half(2 * FilteringWeight)));
 
-		#if CONFIG_MOIRE_DETECTION
-		tsr_simd_halfC MoireErrorSize;
-		UNROLL_N(SIMD_SIZE)
-		for (uint SimdIndex = 0; SimdIndex < SIMD_SIZE; SimdIndex++)
-		{
-			tsr_halfC ClampedEnergy = abs(ClampedFilteredHistory.GetElement(SimdIndex) - FilteredHistory.GetElement(SimdIndex));
-			tsr_halfC ClampedEnergyAmount = ClampedEnergy * SafeRcp(ClampedEnergy[0] + ClampedEnergy[1] + ClampedEnergy[2]);
+			#if CONFIG_MOIRE_DETECTION
+			tsr_simd_halfC MoireErrorSize;
+			UNROLL_N(SIMD_SIZE)
+			for (uint SimdIndex = 0; SimdIndex < SIMD_SIZE; SimdIndex++)
+			{
+				tsr_halfC ClampedEnergy = abs(ClampedFilteredHistory.GetElement(SimdIndex) - FilteredHistory.GetElement(SimdIndex));
+				tsr_halfC ClampedEnergyAmount = ClampedEnergy * SafeRcp(ClampedEnergy[0] + ClampedEnergy[1] + ClampedEnergy[2]);
 
-			// Apply the moire error on each individual channel based on how much they need to clamp.
-			MoireErrorSize.SetElement(SimdIndex, ClampedEnergyAmount * ((FilteringWeight * tsr_half(3.0)) * MoireError.GetElement(SimdIndex)));
+				// Apply the moire error on each individual channel based on how much they need to clamp.
+				MoireErrorSize.SetElement(SimdIndex, ClampedEnergyAmount * ((FilteringWeight * tsr_half(3.0)) * MoireError.GetElement(SimdIndex)));
+			}
+			BoxSize = max(BoxSize, MoireErrorSize);
+			#endif
 		}
-		BoxSize = max(BoxSize, MoireErrorSize);
-		#endif
+
+		tsr_simd_halfC RawClampedEnergy      = abs(ClampedFilteredHistory - FilteredHistory);
+		tsr_simd_halfC FilteredClampedEnergy = Max3x3(Median3x3(RawClampedEnergy));
 
 		tsr_simd_halfC Delta = abs(FilteredInput - FilteredHistory);
 		tsr_simd_halfC ClampError = max(
@@ -580,9 +553,7 @@ void MainCS(
 		tsr_simd_half FilteredRejection = min3(FilteredFactor[0], FilteredFactor[1], FilteredFactor[2]);
 
 		RejectionBlendFinal = saturate(FilteredRejection);
-		
-		RejectionClampBlend = Median3x3(RawRejection);
-		RejectionClampBlend = Min3x3(RejectionClampBlend);
+		RejectionClampBlend = Min3x3(Median3x3(RawRejection));
 	}
 	#endif
 	
diff --git a/Engine/Shaders/Public/LaneVectorization.ush b/Engine/Shaders/Public/LaneVectorization.ush
--- a/Engine/Shaders/Public/LaneVectorization.ush
+++ b/Engine/Shaders/Public/LaneVectorization.ush
@@ -8,6 +8,149 @@
 
 #include "Platform.ush"
 
+bool not(bool x)
+{
+	return !x;
+}
+
+bool2 not(bool2 x)
+{
+	return !x;
+}
+
+bool2 v_pack_b32_b16(bool a, bool b)
+{
+	return bool2(a, b);
+}
+
+float2 v_pack_b32_b16(float a, float b)
+{
+	return float2(a, b);
+}
+
+uint2 v_pack_b32_b16(uint a, uint b)
+{
+	return uint2(a, b);
+}
+
+int2 v_pack_b32_b16(int a, int b)
+{
+	return int2(a, b);
+}
+
+float fast_sign(float x)
+{
+    return clamp(x * asfloat(0x7f7fffff), float(-1.0), float(1.0));
+}
+
+float2 fast_sign(float2 x)
+{
+    return clamp(x * asfloat(0x7f7fffff), float(-1.0), float(1.0));
+}
+
+#if PLATFORM_SUPPORTS_REAL_TYPES
+
+half fast_sign(half x)
+{
+    return clamp(x * asfloat16(uint16_t(0x7bff)), half(-1.0), half(1.0));
+}
+
+half2 fast_sign(half2 x)
+{
+    return clamp(x * asfloat16(uint16_t(0x7bff)), half(-1.0), half(1.0));
+}
+
+#endif
+
+uint bit_and(uint a, uint b)
+{
+    return a & b;
+}
+
+uint2 bit_and(uint2 a, uint2 b)
+{
+    return a & b;
+}
+
+uint bit_or(uint a, uint b)
+{
+    return a | b;
+}
+
+uint2 bit_or(uint2 a, uint2 b)
+{
+    return a | b;
+}
+
+#if PLATFORM_SUPPORTS_REAL_TYPES
+
+uint16_t bit_and(uint16_t a, uint16_t b)
+{
+    return a & b;
+}
+
+uint16_t2 bit_and(uint16_t2 a, uint16_t2 b)
+{
+    return a & b;
+}
+
+uint16_t bit_or(uint16_t a, uint16_t b)
+{
+    return a | b;
+}
+
+uint16_t2 bit_or(uint16_t2 a, uint16_t2 b)
+{
+    return a | b;
+}
+
+#endif
+
+uint bit_shift_left(uint a, uint b)
+{
+    return a << b;
+}
+
+uint2 bit_shift_left(uint2 a, uint2 b)
+{
+    return a << b;
+}
+
+uint bit_shift_right(uint a, uint b)
+{
+    return a >> b;
+}
+
+uint2 bit_shift_right(uint2 a, uint2 b)
+{
+    return a >> b;
+}
+
+#if PLATFORM_SUPPORTS_REAL_TYPES
+
+uint16_t bit_shift_left(uint16_t a, uint16_t b)
+{
+    return a << b;
+}
+
+uint16_t2 bit_shift_left(uint16_t2 a, uint16_t2 b)
+{
+    return a << b;
+}
+
+uint16_t bit_shift_right(uint16_t a, uint16_t b)
+{
+    return a >> b;
+}
+
+uint16_t2 bit_shift_right(uint16_t2 a, uint16_t2 b)
+{
+    return a >> b;
+}
+
+#endif
+
+
 template<typename ScalarType, uint InVectorSize, uint InSimdSize>
 struct TLaneVector
 {
@@ -221,6 +364,17 @@ struct TLaneVector2D
 		}
 	}
 	
+	// Set component of the vector
+	CALL_SITE_DEBUGLOC
+	void SetComponent(uint ComponentIndex, TLaneVector2D<ScalarType, 1, SimdSizeX, SimdSizeY> Comp)
+	{
+		UNROLL
+		for (uint SimdIndex = 0; SimdIndex < SimdSize; SimdIndex++)
+		{
+			RegisterArray[SimdIndex][ComponentIndex] = Comp.RegisterArray[SimdIndex];
+		}
+	}
+
 	
 	// ------------- binary operator +
 	CALL_SITE_DEBUGLOC
@@ -312,6 +466,31 @@ struct TLaneVector2D
 		return R;
 	}
 	
+	// ------------- comparison operators
+	CALL_SITE_DEBUGLOC
+	TLaneVector2D<bool, VectorSize, SimdSizeX, SimdSizeY> operator > (TLaneVector2D<ScalarType, VectorSize, SimdSizeX, SimdSizeY> B)
+	{
+		TLaneVector2D<bool, VectorSize, SimdSizeX, SimdSizeY> R;
+		UNROLL
+		for (uint SimdIndex = 0; SimdIndex < SimdSize; SimdIndex++)
+		{
+			R.SetElement(SimdIndex, GetElement(SimdIndex) > B.GetElement(SimdIndex));
+		}
+		return R;
+	}
+	
+	CALL_SITE_DEBUGLOC
+	TLaneVector2D<bool, VectorSize, SimdSizeX, SimdSizeY> operator < (TLaneVector2D<ScalarType, VectorSize, SimdSizeX, SimdSizeY> B)
+	{
+		TLaneVector2D<bool, VectorSize, SimdSizeX, SimdSizeY> R;
+		UNROLL
+		for (uint SimdIndex = 0; SimdIndex < SimdSize; SimdIndex++)
+		{
+			R.SetElement(SimdIndex, GetElement(SimdIndex) < B.GetElement(SimdIndex));
+		}
+		return R;
+	}
+
 
 	// ------------- constructors
 	CALL_SITE_DEBUGLOC
@@ -339,18 +518,303 @@ struct TLaneVector2D
 		}
 		return R;
 	}
+	
+	CALL_SITE_DEBUGLOC
+	template<typename SourceScalarType>
+	static TLaneVector2D<ScalarType, VectorSize, SimdSizeX, SimdSizeY> CastFrom(
+		TLaneVector2D<SourceScalarType, 1, SimdSizeX, SimdSizeY> A)
+	{
+		TLaneVector2D<ScalarType, VectorSize, SimdSizeX, SimdSizeY> R;
+		UNROLL
+		for (uint SimdIndex = 0; SimdIndex < SimdSize; SimdIndex++)
+		{
+			R.SetElement(SimdIndex, vector<ScalarType, VectorSize>(A.GetElement(SimdIndex)));
+		}
+		return R;
+	}
+	
+	CALL_SITE_DEBUGLOC
+	static TLaneVector2D<ScalarType, VectorSize, SimdSizeX, SimdSizeY> ForceRegisterPacking(
+		TLaneVector2D<ScalarType, VectorSize, SimdSizeX, SimdSizeY> A)
+	{
+		return A;
+	}
+	
+}; // TLaneVector2D
+
+/** Work exactly like a TLaneVector2D, except values are packed vector<ScalarType, 2> to better hint the shader compiler to use 16bit packed instructions.  */
+template<typename ScalarType, uint InVectorSize, uint InSimdSizeX, uint InSimdSizeY>
+struct TLanePackedVector2D
+{
+	static const uint VectorSize = InVectorSize;
+	static const uint SimdSize = InSimdSizeX * InSimdSizeY;
+	static const uint SimdSizeX = InSimdSizeX;
+	static const uint SimdSizeY = InSimdSizeY;
+	static const uint ScalarPairSize = (VectorSize * SimdSize) / 2u;
+
+	vector<ScalarType, 2u> RegisterArray[VectorSize][SimdSize / 2u];
+
+	
+	// Access a component of the vector
+	CALL_SITE_DEBUGLOC
+	TLanePackedVector2D<ScalarType, 1, SimdSizeX, SimdSizeY> operator [](const uint ComponentIndex)
+	{
+		TLanePackedVector2D<ScalarType, 1, SimdSizeX, SimdSizeY> R;
+		UNROLL
+		for (uint DestScalarPairIndex = 0; DestScalarPairIndex < (SimdSizeX * SimdSizeY) / 2u; DestScalarPairIndex++)
+		{
+			R.SetScalarPair(DestScalarPairIndex, GetScalarPair(DestScalarPairIndex * VectorSize + ComponentIndex));
+		}
+		return R;
+	}
+	
+	// Access an individual element of the vector.
+	CALL_SITE_DEBUGLOC
+	vector<ScalarType, VectorSize> GetElement(const uint SimdIndex)
+	{
+		vector<ScalarType, VectorSize> Element;
+		UNROLL
+		for (uint ComponentIndex = 0; ComponentIndex < VectorSize; ComponentIndex++)
+		{
+			Element[ComponentIndex] = RegisterArray[ComponentIndex][SimdIndex / 2u][SimdIndex % 2u];
+		}
+		return Element;
+	}
+	
+	CALL_SITE_DEBUGLOC
+	void SetElement(const uint SimdIndex, vector<ScalarType, VectorSize> Element)
+	{
+		UNROLL
+		for (uint ComponentIndex = 0; ComponentIndex < VectorSize; ComponentIndex++)
+		{
+			RegisterArray[ComponentIndex][SimdIndex / 2u][SimdIndex % 2u] = Element[ComponentIndex];
+		}
+	}
+	
+	CALL_SITE_DEBUGLOC
+	void SetAllElements(vector<ScalarType, VectorSize> Element)
+	{
+		UNROLL
+		for (uint SimdIndex = 0; SimdIndex < SimdSize; SimdIndex++)
+		{
+			SetElement(SimdIndex, Element);
+		}
+	}
+	
+	// Access an individual scalair pair.
+	CALL_SITE_DEBUGLOC
+	vector<ScalarType, 2> GetScalarPair(const uint ScalarPairIndex)
+	{
+		return RegisterArray[ScalarPairIndex % VectorSize][ScalarPairIndex / VectorSize];
+	}
+	
+	CALL_SITE_DEBUGLOC
+	void SetScalarPair(const uint ScalarPairIndex, vector<ScalarType, 2> ScalarPair)
+	{
+		RegisterArray[ScalarPairIndex % VectorSize][ScalarPairIndex / VectorSize] = ScalarPair;
+	}
+	
+	// Set component of the vector
+	CALL_SITE_DEBUGLOC
+	void SetComponent(uint ComponentIndex, TLanePackedVector2D<ScalarType, 1, SimdSizeX, SimdSizeY> Comp)
+	{
+		UNROLL
+		for (uint SrcPairIndex = 0; SrcPairIndex < (SimdSizeX * SimdSizeY) / 2u; SrcPairIndex++)
+		{
+			SetScalarPair(SrcPairIndex * VectorSize + ComponentIndex, Comp.GetScalarPair(SrcPairIndex));
+		}
+	}
+	
+	// ------------- binary operator +
+	CALL_SITE_DEBUGLOC
+	TLanePackedVector2D<ScalarType, VectorSize, SimdSizeX, SimdSizeY> operator + (TLanePackedVector2D<ScalarType, VectorSize, SimdSizeX, SimdSizeY> B)
+	{
+		TLanePackedVector2D<ScalarType, VectorSize, SimdSizeX, SimdSizeY> R;
+		UNROLL
+		for (uint ScalarPairIndex = 0; ScalarPairIndex < ScalarPairSize; ScalarPairIndex++)
+		{
+			R.SetScalarPair(ScalarPairIndex, GetScalarPair(ScalarPairIndex) + B.GetScalarPair(ScalarPairIndex));
+		}
+		return R;
+	}
+	
+	CALL_SITE_DEBUGLOC
+	TLanePackedVector2D<ScalarType, VectorSize, SimdSizeX, SimdSizeY> operator + (vector<ScalarType, VectorSize> B)
+	{
+		TLanePackedVector2D<ScalarType, VectorSize, SimdSizeX, SimdSizeY> R;
+		UNROLL
+		for (uint ScalarPairIndex = 0; ScalarPairIndex < ScalarPairSize; ScalarPairIndex++)
+		{
+			R.SetScalarPair(ScalarPairIndex, GetScalarPair(ScalarPairIndex) + B[ScalarPairIndex % VectorSize]);
+		}
+		return R;
+	}
+
+	
+	// ------------- unary operator -
+	CALL_SITE_DEBUGLOC
+	TLanePackedVector2D<ScalarType, VectorSize, SimdSizeX, SimdSizeY> operator - ()
+	{
+		TLanePackedVector2D<ScalarType, VectorSize, SimdSizeX, SimdSizeY> R;
+		UNROLL
+		for (uint ScalarPairIndex = 0; ScalarPairIndex < ScalarPairSize; ScalarPairIndex++)
+		{
+			R.SetScalarPair(ScalarPairIndex, -GetScalarPair(ScalarPairIndex));
+		}
+		return R;
+	}
+	
+	// ------------- binary operator -
+	CALL_SITE_DEBUGLOC
+	TLanePackedVector2D<ScalarType, VectorSize, SimdSizeX, SimdSizeY> operator - (TLanePackedVector2D<ScalarType, VectorSize, SimdSizeX, SimdSizeY> B)
+	{
+		TLanePackedVector2D<ScalarType, VectorSize, SimdSizeX, SimdSizeY> R;
+		UNROLL
+		for (uint ScalarPairIndex = 0; ScalarPairIndex < ScalarPairSize; ScalarPairIndex++)
+		{
+			R.SetScalarPair(ScalarPairIndex, GetScalarPair(ScalarPairIndex) - B.GetScalarPair(ScalarPairIndex));
+		}
+		return R;
+	}
+	
+	CALL_SITE_DEBUGLOC
+	TLanePackedVector2D<ScalarType, VectorSize, SimdSizeX, SimdSizeY> operator - (vector<ScalarType, VectorSize> B)
+	{
+		TLanePackedVector2D<ScalarType, VectorSize, SimdSizeX, SimdSizeY> R;
+		UNROLL
+		for (uint ScalarPairIndex = 0; ScalarPairIndex < ScalarPairSize; ScalarPairIndex++)
+		{
+			R.SetScalarPair(ScalarPairIndex, GetScalarPair(ScalarPairIndex) - B[ScalarPairIndex % VectorSize]);
+		}
+		return R;
+	}
+
+	
+	// ------------- binary operator *
+	CALL_SITE_DEBUGLOC
+	TLanePackedVector2D<ScalarType, VectorSize, SimdSizeX, SimdSizeY> operator * (TLanePackedVector2D<ScalarType, VectorSize, SimdSizeX, SimdSizeY> B)
+	{
+		TLanePackedVector2D<ScalarType, VectorSize, SimdSizeX, SimdSizeY> R;
+		UNROLL
+		for (uint ScalarPairIndex = 0; ScalarPairIndex < ScalarPairSize; ScalarPairIndex++)
+		{
+			R.SetScalarPair(ScalarPairIndex, GetScalarPair(ScalarPairIndex) * B.GetScalarPair(ScalarPairIndex));
+		}
+		return R;
+	}
+
+	CALL_SITE_DEBUGLOC
+	TLanePackedVector2D<ScalarType, VectorSize, SimdSizeX, SimdSizeY> operator * (vector<ScalarType, VectorSize> B)
+	{
+		TLanePackedVector2D<ScalarType, VectorSize, SimdSizeX, SimdSizeY> R;
+		UNROLL
+		for (uint ScalarPairIndex = 0; ScalarPairIndex < ScalarPairSize; ScalarPairIndex++)
+		{
+			R.SetScalarPair(ScalarPairIndex, GetScalarPair(ScalarPairIndex) * B[ScalarPairIndex % VectorSize]);
+		}
+		return R;
+	}
+	
+	// ------------- comparison operators
+	CALL_SITE_DEBUGLOC
+	TLanePackedVector2D<bool, VectorSize, SimdSizeX, SimdSizeY> operator > (TLanePackedVector2D<ScalarType, VectorSize, SimdSizeX, SimdSizeY> B)
+	{
+		TLanePackedVector2D<bool, VectorSize, SimdSizeX, SimdSizeY> R;
+		UNROLL
+		for (uint ScalarPairIndex = 0; ScalarPairIndex < ScalarPairSize; ScalarPairIndex++)
+		{
+			R.SetScalarPair(ScalarPairIndex, GetScalarPair(ScalarPairIndex) > B.GetScalarPair(ScalarPairIndex));
+		}
+		return R;
+	}
+	
+	CALL_SITE_DEBUGLOC
+	TLanePackedVector2D<bool, VectorSize, SimdSizeX, SimdSizeY> operator < (TLanePackedVector2D<ScalarType, VectorSize, SimdSizeX, SimdSizeY> B)
+	{
+		TLanePackedVector2D<bool, VectorSize, SimdSizeX, SimdSizeY> R;
+		UNROLL
+		for (uint ScalarPairIndex = 0; ScalarPairIndex < ScalarPairSize; ScalarPairIndex++)
+		{
+			R.SetScalarPair(ScalarPairIndex, GetScalarPair(ScalarPairIndex) < B.GetScalarPair(ScalarPairIndex));
+		}
+		return R;
+	}
+
+	// ------------- constructors
+	CALL_SITE_DEBUGLOC
+	static TLanePackedVector2D<ScalarType, VectorSize, SimdSizeX, SimdSizeY> Vectorize(
+		TLanePackedVector2D<ScalarType, 1, SimdSizeX, SimdSizeY> A)
+	{
+		TLanePackedVector2D<ScalarType, VectorSize, SimdSizeX, SimdSizeY> R;
+		UNROLL
+		for (uint SimdIndex = 0; SimdIndex < SimdSize; SimdIndex++)
+		{
+			R.SetElement(SimdIndex, A.GetElement(SimdIndex));
+		}
+		return R;
+	}
+	
+	CALL_SITE_DEBUGLOC
+	static TLanePackedVector2D<ScalarType, VectorSize, SimdSizeX, SimdSizeY> Const(
+		vector<ScalarType, VectorSize> A)
+	{
+		TLanePackedVector2D<ScalarType, VectorSize, SimdSizeX, SimdSizeY> R;
+		R.SetAllElements(A);
+		return R;
+	}
+	
+	CALL_SITE_DEBUGLOC
+	template<typename SourceScalarType>
+	static TLanePackedVector2D<ScalarType, VectorSize, SimdSizeX, SimdSizeY> CastFrom(
+		TLanePackedVector2D<SourceScalarType, VectorSize, SimdSizeX, SimdSizeY> A)
+	{
+		TLanePackedVector2D<ScalarType, VectorSize, SimdSizeX, SimdSizeY> R;
+		UNROLL
+		for (uint ScalarPairIndex = 0; ScalarPairIndex < ScalarPairSize; ScalarPairIndex++)
+		{
+			R.SetScalarPair(ScalarPairIndex, vector<ScalarType, 2>(A.GetScalarPair(ScalarPairIndex)));
+		}
+		return R;
+	}
+	
+	CALL_SITE_DEBUGLOC
+	static TLanePackedVector2D<ScalarType, VectorSize, SimdSizeX, SimdSizeY> ForceRegisterPacking(
+		TLanePackedVector2D<ScalarType, VectorSize, SimdSizeX, SimdSizeY> A)
+	{
+		TLanePackedVector2D<ScalarType, VectorSize, SimdSizeX, SimdSizeY> R;
+		UNROLL
+		for (uint ScalarPairIndex = 0; ScalarPairIndex < ScalarPairSize; ScalarPairIndex++)
+		{
+			vector<ScalarType, 2> Pair = A.GetScalarPair(ScalarPairIndex);
+			R.SetScalarPair(ScalarPairIndex, v_pack_b32_b16(Pair.x, Pair.y));
+		}
+		return R;
+	}
+	
+	CALL_SITE_DEBUGLOC
+	static TLanePackedVector2D<ScalarType, VectorSize, SimdSizeX, SimdSizeY> ChangeRegisterOrder(
+		TLaneVector2D<ScalarType, VectorSize, SimdSizeX, SimdSizeY> A)
+	{
+		TLanePackedVector2D<ScalarType, VectorSize, SimdSizeX, SimdSizeY> R;
+		UNROLL
+		for (uint SimdIndex = 0; SimdIndex < SimdSize; SimdIndex++)
+		{
+			R.SetElement(SimdIndex, A.GetElement(SimdIndex));
+		}
+		return R;
+	}
+}; // struct TLanePackedVector2D
 
-};
 
 // ------------- 1 parameter
 
-#define TVECTOR_FUNCTION_1PARAMS(FunctionName) \
+#define TVECTOR_FUNCTION_1PARAMS(ReturnScalarType, FunctionName) \
 	CALL_SITE_DEBUGLOC \
 	template<typename ScalarType, uint VectorSize, uint SimdSize> \
-	TLaneVector<ScalarType, VectorSize, SimdSize> FunctionName( \
+	TLaneVector<ReturnScalarType, VectorSize, SimdSize> FunctionName( \
 		TLaneVector<ScalarType, VectorSize, SimdSize> A) \
 	{ \
-		TLaneVector<ScalarType, VectorSize, SimdSize> R; \
+		TLaneVector<ReturnScalarType, VectorSize, SimdSize> R; \
 		UNROLL \
 		for (uint SimdIndex = 0; SimdIndex < SimdSize; SimdIndex++) \
 		{ \
@@ -361,10 +825,10 @@ struct TLaneVector2D
 	\
 	CALL_SITE_DEBUGLOC \
 	template<typename ScalarType, uint VectorSize, uint SimdSizeX, uint SimdSizeY> \
-	TLaneVector2D<ScalarType, VectorSize, SimdSizeX, SimdSizeY> FunctionName( \
+	TLaneVector2D<ReturnScalarType, VectorSize, SimdSizeX, SimdSizeY> FunctionName( \
 		TLaneVector2D<ScalarType, VectorSize, SimdSizeX, SimdSizeY> A) \
 	{ \
-		TLaneVector2D<ScalarType, VectorSize, SimdSizeX, SimdSizeY> R; \
+		TLaneVector2D<ReturnScalarType, VectorSize, SimdSizeX, SimdSizeY> R; \
 		UNROLL \
 		for (uint SimdIndex = 0; SimdIndex < SimdSizeX * SimdSizeY; SimdIndex++) \
 		{ \
@@ -372,14 +836,91 @@ struct TLaneVector2D
 		} \
 		return R; \
 	} \
+	\
+	CALL_SITE_DEBUGLOC \
+	template<typename ScalarType, uint VectorSize, uint SimdSizeX, uint SimdSizeY> \
+	TLanePackedVector2D<ReturnScalarType, VectorSize, SimdSizeX, SimdSizeY> FunctionName( \
+		TLanePackedVector2D<ScalarType, VectorSize, SimdSizeX, SimdSizeY> A) \
+	{ \
+		TLanePackedVector2D<ReturnScalarType, VectorSize, SimdSizeX, SimdSizeY> R; \
+		UNROLL \
+		for (uint ScalarPairIndex = 0; ScalarPairIndex < (VectorSize * SimdSizeX * SimdSizeY) / 2; ScalarPairIndex++) \
+		{ \
+			R.SetScalarPair(ScalarPairIndex, FunctionName(A.GetScalarPair(ScalarPairIndex))); \
+		} \
+		return R; \
+	} \
 	
-TVECTOR_FUNCTION_1PARAMS(rcp);
-TVECTOR_FUNCTION_1PARAMS(saturate);
-TVECTOR_FUNCTION_1PARAMS(abs);
+TVECTOR_FUNCTION_1PARAMS(ScalarType, rcp);
+TVECTOR_FUNCTION_1PARAMS(ScalarType, saturate);
+TVECTOR_FUNCTION_1PARAMS(ScalarType, abs);
+TVECTOR_FUNCTION_1PARAMS(ScalarType, floor);
+TVECTOR_FUNCTION_1PARAMS(ScalarType, ceil);
+TVECTOR_FUNCTION_1PARAMS(ScalarType, round);
+TVECTOR_FUNCTION_1PARAMS(ScalarType, fast_sign);
+TVECTOR_FUNCTION_1PARAMS(bool, not);
+TVECTOR_FUNCTION_1PARAMS(uint, asuint);
+#if PLATFORM_SUPPORTS_REAL_TYPES
+	TVECTOR_FUNCTION_1PARAMS(uint16_t, asuint16);
+#endif
+
 
 // ------------- 2 parameters
 
-#define TVECTOR_FUNCTION_2PARAMS(ReturnVectorSize, FunctionName) \
+#define TVECTOR_FUNCTION_2PARAMS(FunctionName) \
+	CALL_SITE_DEBUGLOC \
+	template<typename ScalarType, uint VectorSize, uint SimdSize> \
+	TLaneVector<ScalarType, VectorSize, SimdSize> FunctionName( \
+		TLaneVector<ScalarType, VectorSize, SimdSize> A, \
+		TLaneVector<ScalarType, VectorSize, SimdSize> B) \
+	{ \
+		TLaneVector<ScalarType, VectorSize, SimdSize> R; \
+		UNROLL for (uint SimdIndex = 0; SimdIndex < SimdSize; SimdIndex++) { \
+			R.SetElement(SimdIndex, FunctionName(A.GetElement(SimdIndex), B.GetElement(SimdIndex))); \
+		} \
+		return R; \
+	} \
+	\
+	CALL_SITE_DEBUGLOC \
+	template<typename ScalarType, uint VectorSize, uint SimdSizeX, uint SimdSizeY> \
+	TLaneVector2D<ScalarType, VectorSize, SimdSizeX, SimdSizeY> FunctionName( \
+		TLaneVector2D<ScalarType, VectorSize, SimdSizeX, SimdSizeY> A, \
+		TLaneVector2D<ScalarType, VectorSize, SimdSizeX, SimdSizeY> B) \
+	{ \
+		TLaneVector2D<ScalarType, VectorSize, SimdSizeX, SimdSizeY> R; \
+		UNROLL for (uint SimdIndex = 0; SimdIndex < SimdSizeX * SimdSizeY; SimdIndex++) { \
+			R.SetElement(SimdIndex, FunctionName(A.GetElement(SimdIndex), B.GetElement(SimdIndex))); \
+		} \
+		return R; \
+	} \
+	\
+	CALL_SITE_DEBUGLOC \
+	template<typename ScalarType, uint VectorSize, uint SimdSizeX, uint SimdSizeY> \
+	TLanePackedVector2D<ScalarType, VectorSize, SimdSizeX, SimdSizeY> FunctionName( \
+		TLanePackedVector2D<ScalarType, VectorSize, SimdSizeX, SimdSizeY> A, \
+		TLanePackedVector2D<ScalarType, VectorSize, SimdSizeX, SimdSizeY> B) \
+	{ \
+		TLanePackedVector2D<ScalarType, VectorSize, SimdSizeX, SimdSizeY> R; \
+		UNROLL for (uint ScalarPairIndex = 0; ScalarPairIndex < (VectorSize * SimdSizeX * SimdSizeY) / 2u; ScalarPairIndex++) \
+		{ \
+			R.SetScalarPair(ScalarPairIndex, FunctionName(A.GetScalarPair(ScalarPairIndex), B.GetScalarPair(ScalarPairIndex))); \
+		} \
+		return R; \
+	} \
+
+TVECTOR_FUNCTION_2PARAMS(min);
+TVECTOR_FUNCTION_2PARAMS(max);
+TVECTOR_FUNCTION_2PARAMS(and_internal);
+TVECTOR_FUNCTION_2PARAMS(or_internal);
+TVECTOR_FUNCTION_2PARAMS(bit_and);
+TVECTOR_FUNCTION_2PARAMS(bit_or);
+TVECTOR_FUNCTION_2PARAMS(bit_shift_left);
+TVECTOR_FUNCTION_2PARAMS(bit_shift_right);
+
+
+// ------------- 2 parameters but different return dimension
+
+#define TVECTOR_FUNCTION_2PARAMS_DIFF_RETURN(ReturnVectorSize, FunctionName) \
 	CALL_SITE_DEBUGLOC \
 	template<typename ScalarType, uint VectorSize, uint SimdSize> \
 	TLaneVector<ScalarType, ReturnVectorSize, SimdSize> FunctionName( \
@@ -405,22 +946,33 @@ TVECTOR_FUNCTION_1PARAMS(abs);
 		} \
 		return R; \
 	} \
+	\
+	CALL_SITE_DEBUGLOC \
+	template<typename ScalarType, uint VectorSize, uint SimdSizeX, uint SimdSizeY> \
+	TLanePackedVector2D<ScalarType, ReturnVectorSize, SimdSizeX, SimdSizeY> FunctionName( \
+		TLanePackedVector2D<ScalarType, VectorSize, SimdSizeX, SimdSizeY> A, \
+		TLanePackedVector2D<ScalarType, VectorSize, SimdSizeX, SimdSizeY> B) \
+	{ \
+		TLanePackedVector2D<ScalarType, ReturnVectorSize, SimdSizeX, SimdSizeY> R; \
+		UNROLL for (uint SimdIndex = 0; SimdIndex < SimdSizeX * SimdSizeY; SimdIndex++) { \
+			R.SetElement(SimdIndex, FunctionName(A.GetElement(SimdIndex), B.GetElement(SimdIndex))); \
+		} \
+		return R; \
+	} \
 	
-TVECTOR_FUNCTION_2PARAMS(VectorSize, min);
-TVECTOR_FUNCTION_2PARAMS(VectorSize, max);
-TVECTOR_FUNCTION_2PARAMS(/* VectorSize = */ 1, dot);
+TVECTOR_FUNCTION_2PARAMS_DIFF_RETURN(/* VectorSize = */ 1, dot);
 
 
 // ------------- 3 parameters
-#define TVECTOR_FUNCTION_3PARAMS(ReturnVectorSize, FunctionName) \
+#define TVECTOR_FUNCTION_3PARAMS(FunctionName, ScalarTypeA) \
 	CALL_SITE_DEBUGLOC \
 	template<typename ScalarType, uint VectorSize, uint SimdSize> \
-	TLaneVector<ScalarType, ReturnVectorSize, SimdSize> FunctionName( \
-		TLaneVector<ScalarType, VectorSize, SimdSize> A, \
+	TLaneVector<ScalarType, VectorSize, SimdSize> FunctionName( \
+		TLaneVector<ScalarTypeA, VectorSize, SimdSize> A, \
 		TLaneVector<ScalarType, VectorSize, SimdSize> B, \
 		TLaneVector<ScalarType, VectorSize, SimdSize> C) \
 	{ \
-		TLaneVector<ScalarType, ReturnVectorSize, SimdSize> R; \
+		TLaneVector<ScalarType, VectorSize, SimdSize> R; \
 		UNROLL for (uint SimdIndex = 0; SimdIndex < SimdSize; SimdIndex++) { \
 			R.SetElement(SimdIndex, FunctionName(A.GetElement(SimdIndex), B.GetElement(SimdIndex), C.GetElement(SimdIndex))); \
 		} \
@@ -429,21 +981,37 @@ TVECTOR_FUNCTION_2PARAMS(/* VectorSize = */ 1, dot);
 	\
 	CALL_SITE_DEBUGLOC \
 	template<typename ScalarType, uint VectorSize, uint SimdSizeX, uint SimdSizeY> \
-	TLaneVector2D<ScalarType, ReturnVectorSize, SimdSizeX, SimdSizeY> FunctionName( \
-		TLaneVector2D<ScalarType, VectorSize, SimdSizeX, SimdSizeY> A, \
+	TLaneVector2D<ScalarType, VectorSize, SimdSizeX, SimdSizeY> FunctionName( \
+		TLaneVector2D<ScalarTypeA, VectorSize, SimdSizeX, SimdSizeY> A, \
 		TLaneVector2D<ScalarType, VectorSize, SimdSizeX, SimdSizeY> B, \
 		TLaneVector2D<ScalarType, VectorSize, SimdSizeX, SimdSizeY> C) \
 	{ \
-		TLaneVector2D<ScalarType, ReturnVectorSize, SimdSizeX, SimdSizeY> R; \
+		TLaneVector2D<ScalarType, VectorSize, SimdSizeX, SimdSizeY> R; \
 		UNROLL for (uint SimdIndex = 0; SimdIndex < SimdSizeX * SimdSizeY; SimdIndex++) { \
 			R.SetElement(SimdIndex, FunctionName(A.GetElement(SimdIndex), B.GetElement(SimdIndex), C.GetElement(SimdIndex))); \
 		} \
 		return R; \
 	} \
+	\
+	CALL_SITE_DEBUGLOC \
+	template<typename ScalarType, uint VectorSize, uint SimdSizeX, uint SimdSizeY> \
+	TLanePackedVector2D<ScalarType, VectorSize, SimdSizeX, SimdSizeY> FunctionName( \
+		TLanePackedVector2D<ScalarTypeA, VectorSize, SimdSizeX, SimdSizeY> A, \
+		TLanePackedVector2D<ScalarType, VectorSize, SimdSizeX, SimdSizeY> B, \
+		TLanePackedVector2D<ScalarType, VectorSize, SimdSizeX, SimdSizeY> C) \
+	{ \
+		TLanePackedVector2D<ScalarType, VectorSize, SimdSizeX, SimdSizeY> R; \
+		UNROLL for (uint ScalarPairIndex = 0; ScalarPairIndex < (VectorSize * SimdSizeX * SimdSizeY) / 2u; ScalarPairIndex++) \
+		{ \
+			R.SetScalarPair(ScalarPairIndex, FunctionName(A.GetScalarPair(ScalarPairIndex), B.GetScalarPair(ScalarPairIndex), C.GetScalarPair(ScalarPairIndex))); \
+		} \
+		return R; \
+	} \
 	
-TVECTOR_FUNCTION_3PARAMS(VectorSize, clamp);
-TVECTOR_FUNCTION_3PARAMS(VectorSize, min3);
-TVECTOR_FUNCTION_3PARAMS(VectorSize, max3);
+TVECTOR_FUNCTION_3PARAMS(select_internal, bool);
+TVECTOR_FUNCTION_3PARAMS(clamp, ScalarType);
+TVECTOR_FUNCTION_3PARAMS(min3, ScalarType);
+TVECTOR_FUNCTION_3PARAMS(max3, ScalarType);
 #if COMPILER_SUPPORTS_MED3
-	TVECTOR_FUNCTION_3PARAMS(VectorSize, med3);
+	TVECTOR_FUNCTION_3PARAMS(med3, ScalarType);
 #endif
