VulkanRHI: Expose UB_MEMBER_UNIFORM_BLOCK macro in bindless also, add the option to override it per platform.  Fix shader compilation in Vulkan bindless.

diff --git a/Engine/Shaders/Public/BindlessResources.ush b/Engine/Shaders/Public/BindlessResources.ush
--- a/Engine/Shaders/Public/BindlessResources.ush
+++ b/Engine/Shaders/Public/BindlessResources.ush
@@ -46,6 +46,9 @@
 	#define GET_RESOURCE(Name)                    Name
 
 	#define UB_RESOURCE_MEMBER_RESOURCE(Type, Name) Type Name
+#endif
+
+#ifndef UB_MEMBER_UNIFORM_BLOCK
 	#define UB_MEMBER_UNIFORM_BLOCK(Name) float4 Name[PLATFORM_MAX_UNIFORM_BUFFER_RANGE_FLOAT4]
 #endif
 
