- Updated Guid in MetalCommon.ush to work around an issue caused by high level code change that now need metal shaders to be recompiled

diff --git a/Engine/Shaders/Public/Platform/Metal/MetalCommon.ush b/Engine/Shaders/Public/Platform/Metal/MetalCommon.ush
--- a/Engine/Shaders/Public/Platform/Metal/MetalCommon.ush
+++ b/Engine/Shaders/Public/Platform/Metal/MetalCommon.ush
@@ -5,7 +5,7 @@
 =============================================================================*/
 
 // Update this GUID to improve shader recompilation for Metal only shaders
-// GUID = 7DDC23C5-DDBC-4A9F-9047-CC365A54D83C
+// GUID = 56884A3E-F17C-4342-8E97-035ABCE1E692
 
 #pragma once
 
