Added WaveInclusivePrefixSum to Platform.ush

diff --git a/Engine/Shaders/Public/Platform.ush b/Engine/Shaders/Public/Platform.ush
--- a/Engine/Shaders/Public/Platform.ush
+++ b/Engine/Shaders/Public/Platform.ush
@@ -1232,6 +1232,11 @@ uint WaveGetActiveLaneIndexLast()
 
 #endif // PLATFORM_SUPPORTS_SM6_0_WAVE_OPERATIONS
 
+#if !defined(COMPILER_SUPPORTS_WAVE_INCLUSIVE_PREFIX_SUM)
+#define COMPILER_SUPPORTS_WAVE_INCLUSIVE_PREFIX_SUM 0
+#define WaveInclusivePrefixSum(x) (WavePrefixSum(x) + x)
+#endif
+
 // Give hint to compiler to move one value to scalar unit.
 #if !defined(ToScalarMemory) && !defined(COMPILER_SUPPORTS_TO_SCALAR_MEMORY)
 	#define ToScalarMemory(x) (x)
