Implements different register layout swapable at compile time to work arround shader compiler bugs

diff --git a/Engine/Shaders/Private/TemporalSuperResolution/TSRConvolutionNetwork.ush b/Engine/Shaders/Private/TemporalSuperResolution/TSRConvolutionNetwork.ush
--- a/Engine/Shaders/Private/TemporalSuperResolution/TSRConvolutionNetwork.ush
+++ b/Engine/Shaders/Private/TemporalSuperResolution/TSRConvolutionNetwork.ush
@@ -442,8 +442,7 @@ TLaneVector2D<ScalarType, VectorSize, LaneStrideX, LaneStrideY> WaveAccessNeighb
 		const int LaneRotation = dot(select(bNeedsLaneRotation, int2(Offset), 0), int2(1, LANE_COUNT_X));
 		
 		// Access the lement.
-		vector<ScalarType, VectorSize> SimdElement = Input.GetElement(NeigborSimdIndex);
-		vector<ScalarType, VectorSize> ReturnSimdElement = SimdElement;
+		vector<ScalarType, VectorSize> ReturnSimdElement;
 		if (LaneRotation != 0)
 		{
 			// Only uses wave intrinsics on platforms that provides RDNA's ds_swizzle because we know the wave size of this platforms. 
@@ -451,12 +450,18 @@ TLaneVector2D<ScalarType, VectorSize, LaneStrideX, LaneStrideY> WaveAccessNeighb
 				//#error InitWaveRotateLaneGroup() isn't supported.
 			#endif
 			const FWaveBroadcastSettings BroadcastSettings = InitWaveRotateLaneGroup(/* LaneGroupSize = */ LANE_COUNT, LaneRotation);
-			ReturnSimdElement = WaveBroadcast(BroadcastSettings, SimdElement);
+			ReturnSimdElement = Input.Registers.WaveBroadcastElement(BroadcastSettings, NeigborSimdIndex);
+		}
+		else
+		{
+			ReturnSimdElement = Input.GetElement(NeigborSimdIndex);
 		}
 
 		Return.SetElement(SimdIndex, ReturnSimdElement);
 	}
 
+	Return.TightenRegisters();
+
 	return Return;
 }
 
@@ -499,6 +504,9 @@ void AccessNeighborTexels3x1(
 		// TODO
 		CP = CN = CO;
 	}
+
+	CP.TightenRegisters();
+	CN.TightenRegisters();
 }
 #else
 {
@@ -574,6 +582,9 @@ void AccessNeighborTexels1x3(
 			CN.SetElement(FirstSimdIndex + SimdIndex, FirstElement);
 		}
 	}
+	
+	CP.TightenRegisters();
+	CN.TightenRegisters();
 }
 #else
 {
diff --git a/Engine/Shaders/Private/TemporalSuperResolution/TSRRejectShading.usf b/Engine/Shaders/Private/TemporalSuperResolution/TSRRejectShading.usf
--- a/Engine/Shaders/Private/TemporalSuperResolution/TSRRejectShading.usf
+++ b/Engine/Shaders/Private/TemporalSuperResolution/TSRRejectShading.usf
@@ -10,7 +10,6 @@
 //------------------------------------------------------- INCLUDES
 
 #include "TSRCommon.ush"
-#include "/Engine/Public/LaneVectorization.ush"
 
 
 //------------------------------------------------------- CONFIG
@@ -63,6 +62,10 @@
 
 	// Each wave needs to output top and bottom rows of pixels into LDS;
 	#define LDS_SIZE (2 * (LANE_STRIDE_X * LANE_COUNT_X) * (WAVE_COUNT_X * WAVE_COUNT_Y))
+	
+	#ifdef PLATFORM_SPECIFIC_TENSOR_REGISTER_LAYOUT
+		#define TENSOR_REGISTER_LAYOUT PLATFORM_SPECIFIC_TENSOR_REGISTER_LAYOUT
+	#endif
 
 #elif DIM_WAVE_SIZE == 0
 	/** Load an entire 32x32 with 1x1 worth of data by lane, only using LDS. */
@@ -78,6 +81,8 @@
 	#define LANE_STRIDE_Y 1
 
 	#define LDS_SIZE (WAVE_COUNT_X * WAVE_COUNT_Y)
+	
+	#define TENSOR_REGISTER_LAYOUT FRegisters_S_OneVector
 
 #else
 	#error Unknown wave size
@@ -310,10 +315,13 @@ void ComputeMoireError(
 		#endif
 		tsr_half LocalMoireError = (abs(LocalNewTotalVariation * SafeRcp(LocalNewTotalVariationCount)) + LocalNewTotalVariationCount * MoireEncodingError) * LocalCountFadeIn;
 
-		OutMoireHistory.RegisterArray[SimdIndex][0] = FinalFlickeringHistory.GetElement(SimdIndex);
-		OutMoireHistory.RegisterArray[SimdIndex][1] = LocalNewGradient * tsr_half(127.0 / 255.0) + tsr_half(127.0 / 255.0);
-		OutMoireHistory.RegisterArray[SimdIndex][2] = LocalNewTotalVariation;
-		OutMoireHistory.RegisterArray[SimdIndex][3] = LocalQuantizedNewTotalVariationCount * rcp(MaxPrevTotalVariationCount);
+		tsr_half4 PackedMoireHistory = tsr_half4(
+			FinalFlickeringHistory.GetElement(SimdIndex),
+			LocalNewGradient * tsr_half(127.0 / 255.0) + tsr_half(127.0 / 255.0),
+			LocalNewTotalVariation,
+			LocalQuantizedNewTotalVariationCount * rcp(MaxPrevTotalVariationCount));
+
+		OutMoireHistory.SetElement(SimdIndex, PackedMoireHistory);
 
 		OutMoireError.SetElement(SimdIndex, LocalMoireError);
 	}
@@ -515,6 +523,24 @@ void MainCS(
 				#endif
 			#endif
 		}
+		
+		OriginalOpaqueInput.TightenRegisters();
+		#if CONFIG_COMPOSE_TRANSLUCENCY
+			OriginalTranslucencyInput.TightenRegisters();
+		#endif
+		OriginalHistory.TightenRegisters();
+		OriginalHistoryUncertainty.TightenRegisters();
+		VelocityEdge.TightenRegisters();
+		bIsValid.TightenRegisters();
+		#if CONFIG_MOIRE_DETECTION
+			MoireInput.TightenRegisters();
+			PrevMoireHistory.TightenRegisters();
+			bIsStatic.TightenRegisters();
+		#endif
+		#if CONFIG_HISTORY_RESURRECTION
+			OriginalResurrectedHistory.TightenRegisters();
+			OriginalResurrectedHistoryUncertainty.TightenRegisters();
+		#endif
 	}
 	
 	// Compute AA Luma
diff --git a/Engine/Shaders/Public/LaneVectorization.ush b/Engine/Shaders/Public/LaneVectorization.ush
--- a/Engine/Shaders/Public/LaneVectorization.ush
+++ b/Engine/Shaders/Public/LaneVectorization.ush
@@ -7,6 +7,21 @@
 #pragma once
 
 #include "Platform.ush"
+#include "WaveBroadcastIntrinsics.ush"
+
+
+//------------------------------------------------------- DEFINE
+
+#ifndef TENSOR_REGISTER_LAYOUT
+	#define TENSOR_REGISTER_LAYOUT FRegisters_AoS_VectorArray
+#endif
+
+#ifndef TENSOR_REGISTER_TIGHTENING
+	#define TENSOR_REGISTER_TIGHTENING 0
+#endif
+
+
+//------------------------------------------------------- STANDARD SCALAR OPERATIONS
 
 bool not(bool x)
 {
@@ -18,6 +33,16 @@ bool2 not(bool2 x)
 	return !x;
 }
 
+bool3 not(bool3 x)
+{
+	return !x;
+}
+
+bool4 not(bool4 x)
+{
+	return !x;
+}
+
 bool2 v_pack_b32_b16(bool a, bool b)
 {
 	return bool2(a, b);
@@ -151,702 +176,914 @@ uint16_t2 bit_shift_right(uint16_t2 a, uint16_t2 b)
 #endif
 
 
-template<typename ScalarType, uint InVectorSize, uint InSimdSize>
-struct TLaneVector
+//------------------------------------------------------- TIGHT REGISTERS
+
+bool  PackRegistersTightly(bool v)
+{
+	return v;
+}
+
+bool2 PackRegistersTightly(bool2 v)
+{
+	return v;
+}
+
+bool3 PackRegistersTightly(bool3 v)
+{
+	return v;
+}
+
+bool4 PackRegistersTightly(bool4 v)
+{
+	return v;
+}
+
+float  PackRegistersTightly(float v)
+{
+	return v;
+}
+
+float2 PackRegistersTightly(float2 v)
+{
+	return v;
+}
+
+float3 PackRegistersTightly(float3 v)
+{
+	return v;
+}
+
+float4 PackRegistersTightly(float4 v)
+{
+	return v;
+}
+
+uint  PackRegistersTightly(uint v)
+{
+	return v;
+}
+
+uint2 PackRegistersTightly(uint2 v)
+{
+	return v;
+}
+
+uint3 PackRegistersTightly(uint3 v)
+{
+	return v;
+}
+
+uint4 PackRegistersTightly(uint4 v)
+{
+	return v;
+}
+
+int  PackRegistersTightly(int v)
+{
+	return v;
+}
+
+int2 PackRegistersTightly(int2 v)
+{
+	return v;
+}
+
+int3 PackRegistersTightly(int3 v)
+{
+	return v;
+}
+
+int4 PackRegistersTightly(int4 v)
+{
+	return v;
+}
+
+
+#if PLATFORM_SUPPORTS_REAL_TYPES
+
+half  PackRegistersTightly(half v)
+{
+	return v;
+}
+
+half2 PackRegistersTightly(half2 v)
+{
+	return v_pack_b32_b16(v.x, v.y);
+}
+
+half3 PackRegistersTightly(half3 v)
+{
+	return half3(v_pack_b32_b16(v.x, v.y), v.z);
+}
+
+half4 PackRegistersTightly(half4 v)
+{
+	return half4(v_pack_b32_b16(v.x, v.y), v_pack_b32_b16(v.z, v.w));
+}
+
+uint16_t  PackRegistersTightly(uint16_t v)
+{
+	return v;
+}
+
+uint16_t2 PackRegistersTightly(uint16_t2 v)
+{
+	return v_pack_b32_b16(v.x, v.y);
+}
+
+uint16_t3 PackRegistersTightly(uint16_t3 v)
+{
+	return uint16_t3(v_pack_b32_b16(v.x, v.y), v.z);
+}
+
+uint16_t4 PackRegistersTightly(uint16_t4 v)
+{
+	return uint16_t4(v_pack_b32_b16(v.x, v.y), v_pack_b32_b16(v.z, v.w));
+}
+
+int16_t  PackRegistersTightly(int16_t v)
+{
+	return v;
+}
+
+int16_t2 PackRegistersTightly(int16_t2 v)
+{
+	return v_pack_b32_b16(v.x, v.y);
+}
+
+int16_t3 PackRegistersTightly(int16_t3 v)
+{
+	return int16_t3(v_pack_b32_b16(v.x, v.y), v.z);
+}
+
+int16_t4 PackRegistersTightly(int16_t4 v)
 {
+	return int16_t4(v_pack_b32_b16(v.x, v.y), v_pack_b32_b16(v.z, v.w));
+}
+
+
+#endif // PLATFORM_SUPPORTS_REAL_TYPES
+
+
+//------------------------------------------------------- REGISTERS LAYOUTS
+
+/** Stores one unique vector<>. InElementCount must be == 1  */
+template<typename ScalarType, uint InVectorSize, uint InElementCount>
+struct FRegisters_S_OneVector
+{
+	/** Size of each element. */
 	static const uint VectorSize = InVectorSize;
-	static const uint SimdSize = InSimdSize;
+	
+	/** Number of elememnts */
+	static const uint ElementCount = 1;
+	
+	/** Size and number of register rows */
+	static const uint RegisterRowSize = VectorSize;
+	static const uint RegisterRowCount = 1;
 
-	vector<ScalarType, VectorSize> RegisterArray[SimdSize];
+	vector<ScalarType, VectorSize> E;
+	
+
+	// ------------- register rows
 
+	CALL_SITE_DEBUGLOC
+	vector<ScalarType, RegisterRowSize> GetRegisterRow(const uint RegisterRowIndex)
+	{
+		return E;
+	}
 	
-	// Access a component of the vector
 	CALL_SITE_DEBUGLOC
-	TLaneVector<ScalarType, 1, SimdSize> operator [](const uint ComponentIndex)
+	void SetRegisterRow(const uint RegisterRowIndex, vector<ScalarType, RegisterRowSize> RegisterRow)
 	{
-		TLaneVector<ScalarType, 1, SimdSize> R;
-		UNROLL
-		for (uint SimdIndex = 0; SimdIndex < SimdSize; SimdIndex++)
-		{
-			R.SetElement(SimdIndex, GetElement(SimdIndex)[ComponentIndex]);
-		}
-		return R;
+		E = RegisterRow;
 	}
 	
-	// Access an individual element of the vector.
+	
+	// ------------- elements
+
 	CALL_SITE_DEBUGLOC
-	vector<ScalarType, VectorSize> GetElement(const uint SimdIndex)
+	vector<ScalarType, VectorSize> GetElement(const uint ElementIndex)
 	{
-		return RegisterArray[SimdIndex];
+		return E;
 	}
 	
 	CALL_SITE_DEBUGLOC
-	void SetElement(const uint SimdIndex, vector<ScalarType, VectorSize> Element)
+	void SetElement(const uint ElementIndex, vector<ScalarType, VectorSize> Element)
+	{
+		E = Element;
+	}
+
+	vector<ScalarType, VectorSize> WaveBroadcastElement(const FWaveBroadcastSettings BroadcastSettings, const uint ElementIndex)
 	{
-		RegisterArray[SimdIndex] = Element;
+		return WaveBroadcast(BroadcastSettings, GetElement(ElementIndex));
 	}
 	
+	
+	// ------------- elements' components
+
 	CALL_SITE_DEBUGLOC
-	void SetAllElements(vector<ScalarType, VectorSize> Element)
+	ScalarType GetElementComponent(const uint ElementIndex, const uint ComponentIndex)
 	{
-		UNROLL
-		for (uint SimdIndex = 0; SimdIndex < SimdSize; SimdIndex++)
-		{
-			RegisterArray[SimdIndex] = Element;
-		}
+		return E[ComponentIndex];
+	}
+
+	CALL_SITE_DEBUGLOC
+	void SetElementComponent(const uint ElementIndex, const uint ComponentIndex, ScalarType Component)
+	{
+		E[ComponentIndex] = Component;
 	}
+}; // struct FRegisters_S_OneVector
+
+/** Stores an array of vector<> as an array of structure in registers.  */
+template<typename ScalarType, uint InVectorSize, uint InElementCount>
+struct FRegisters_AoS_VectorArray
+{
+	/** Size of each element. */
+	static const uint VectorSize = InVectorSize;
 	
+	/** Number of elememnts */
+	static const uint ElementCount = InElementCount;
+	
+	/** Size and number of register rows */
+	static const uint RegisterRowSize = VectorSize;
+	static const uint RegisterRowCount = ElementCount;
+
+	vector<ScalarType, VectorSize> Array[ElementCount];
+	
+	
+	// ------------- register rows
 	
-	// ------------- binary operator +
 	CALL_SITE_DEBUGLOC
-	TLaneVector<ScalarType, VectorSize, SimdSize> operator + (TLaneVector<ScalarType, VectorSize, SimdSize> B)
+	vector<ScalarType, RegisterRowSize> GetRegisterRow(const uint RegisterRowIndex)
 	{
-		TLaneVector<ScalarType, VectorSize, SimdSize> R;
-		UNROLL
-		for (uint SimdIndex = 0; SimdIndex < SimdSize; SimdIndex++)
-		{
-			R.SetElement(SimdIndex, GetElement(SimdIndex) + B.GetElement(SimdIndex));
-		}
-		return R;
+		return Array[RegisterRowIndex];
 	}
 	
 	CALL_SITE_DEBUGLOC
-	TLaneVector<ScalarType, VectorSize, SimdSize> operator + (vector<ScalarType, VectorSize> B)
+	void SetRegisterRow(const uint RegisterRowIndex, vector<ScalarType, RegisterRowSize> RegisterRow)
 	{
-		TLaneVector<ScalarType, VectorSize, SimdSize> R;
-		UNROLL
-		for (uint SimdIndex = 0; SimdIndex < SimdSize; SimdIndex++)
-		{
-			R.SetElement(SimdIndex, GetElement(SimdIndex) + B);
-		}
-		return R;
+		Array[RegisterRowIndex] = RegisterRow;
 	}
+	
 
+	// ------------- elements
 	
-	// ------------- unary operator -
 	CALL_SITE_DEBUGLOC
-	TLaneVector<ScalarType, VectorSize, SimdSize> operator - ()
+	vector<ScalarType, VectorSize> GetElement(const uint ElementIndex)
 	{
-		TLaneVector<ScalarType, VectorSize, SimdSize> R;
-		UNROLL
-		for (uint SimdIndex = 0; SimdIndex < SimdSize; SimdIndex++)
-		{
-			R.SetElement(SimdIndex, -GetElement(SimdIndex));
-		}
-		return R;
+		return Array[ElementIndex];
 	}
 	
-	// ------------- binary operator -
 	CALL_SITE_DEBUGLOC
-	TLaneVector<ScalarType, VectorSize, SimdSize> operator - (TLaneVector<ScalarType, VectorSize, SimdSize> B)
+	void SetElement(const uint ElementIndex, vector<ScalarType, VectorSize> Element)
 	{
-		TLaneVector<ScalarType, VectorSize, SimdSize> R;
-		UNROLL
-		for (uint SimdIndex = 0; SimdIndex < SimdSize; SimdIndex++)
-		{
-			R.SetElement(SimdIndex, GetElement(SimdIndex) - B.GetElement(SimdIndex));
-		}
-		return R;
+		Array[ElementIndex] = Element;
+	}
+
+	vector<ScalarType, VectorSize> WaveBroadcastElement(const FWaveBroadcastSettings BroadcastSettings, const uint ElementIndex)
+	{
+		return WaveBroadcast(BroadcastSettings, GetElement(ElementIndex));
 	}
 	
+
+	// ------------- elements' components
+
 	CALL_SITE_DEBUGLOC
-	TLaneVector<ScalarType, VectorSize, SimdSize> operator - (vector<ScalarType, VectorSize> B)
+	ScalarType GetElementComponent(const uint ElementIndex, const uint ComponentIndex)
 	{
-		TLaneVector<ScalarType, VectorSize, SimdSize> R;
-		UNROLL
-		for (uint SimdIndex = 0; SimdIndex < SimdSize; SimdIndex++)
-		{
-			R.SetElement(SimdIndex, GetElement(SimdIndex) - B);
-		}
-		return R;
+		return Array[ElementIndex][ComponentIndex];
 	}
 
-	
-	// ------------- binary operator *
 	CALL_SITE_DEBUGLOC
-	TLaneVector<ScalarType, VectorSize, SimdSize> operator * (TLaneVector<ScalarType, VectorSize, SimdSize> B)
+	void SetElementComponent(const uint ElementIndex, const uint ComponentIndex, ScalarType Component)
 	{
-		TLaneVector<ScalarType, VectorSize, SimdSize> R;
-		UNROLL
-		for (uint SimdIndex = 0; SimdIndex < SimdSize; SimdIndex++)
-		{
-			R.SetElement(SimdIndex, GetElement(SimdIndex) * B.GetElement(SimdIndex));
-		}
-		return R;
+		Array[ElementIndex][ComponentIndex] = Component;
 	}
+}; // struct FRegisters_AoS_VectorArray
 
+/** Stores an array of vector<> as an array of structure in registers using a matrix<>. */
+template<typename ScalarType, uint InVectorSize, uint InElementCount>
+struct FRegisters_AoS_Matrix
+{
+	/** Size of each element. */
+	static const uint VectorSize = InVectorSize;
+	
+	/** Number of elememnts */
+	static const uint ElementCount = InElementCount;
+	
+	/** Size and number of register rows */
+	static const uint RegisterRowSize = VectorSize;
+	static const uint RegisterRowCount = ElementCount;
+
+	matrix<ScalarType, ElementCount, VectorSize> M;
+	
+	
+	// ------------- register rows
+	
 	CALL_SITE_DEBUGLOC
-	TLaneVector<ScalarType, VectorSize, SimdSize> operator * (vector<ScalarType, VectorSize> B)
+	vector<ScalarType, RegisterRowSize> GetRegisterRow(const uint RegisterRowIndex)
 	{
-		TLaneVector<ScalarType, VectorSize, SimdSize> R;
-		UNROLL
-		for (uint SimdIndex = 0; SimdIndex < SimdSize; SimdIndex++)
-		{
-			R.SetElement(SimdIndex, GetElement(SimdIndex) * B);
-		}
-		return R;
+		return M[RegisterRowIndex];
+	}
+	
+	CALL_SITE_DEBUGLOC
+	void SetRegisterRow(const uint RegisterRowIndex, vector<ScalarType, RegisterRowSize> RegisterRow)
+	{
+		M[RegisterRowIndex] = RegisterRow;
 	}
 	
 
-	// ------------- constructors
+	// ------------- elements
+	
 	CALL_SITE_DEBUGLOC
-	static TLaneVector<ScalarType, VectorSize, SimdSize> Vectorize(
-		TLaneVector<ScalarType, 1, SimdSize> A)
+	vector<ScalarType, VectorSize> GetElement(const uint ElementIndex)
 	{
-		TLaneVector<ScalarType, VectorSize, SimdSize> R;
-		UNROLL
-		for (uint SimdIndex = 0; SimdIndex < SimdSize; SimdIndex++)
-		{
-			R.SetElement(SimdIndex, A.GetElement(SimdIndex));
-		}
-		return R;
+		return M[ElementIndex];
 	}
 	
 	CALL_SITE_DEBUGLOC
-	static TLaneVector<ScalarType, VectorSize, SimdSize> Const(
-		vector<ScalarType, VectorSize> A)
+	void SetElement(const uint ElementIndex, vector<ScalarType, VectorSize> Element)
 	{
-		TLaneVector<ScalarType, VectorSize, SimdSize> R;
-		UNROLL
-		for (uint SimdIndex = 0; SimdIndex < SimdSize; SimdIndex++)
-		{
-			R.SetElement(SimdIndex, A);
-		}
-		return R;
+		M[ElementIndex] = Element;
 	}
 
-};
+	vector<ScalarType, VectorSize> WaveBroadcastElement(const FWaveBroadcastSettings BroadcastSettings, const uint ElementIndex)
+	{
+		return WaveBroadcast(BroadcastSettings, GetElement(ElementIndex));
+	}
+	
+	
+	// ------------- elements' components
 
-template<typename ScalarType, uint InVectorSize, uint InSimdSizeX, uint InSimdSizeY>
-struct TLaneVector2D
+	CALL_SITE_DEBUGLOC
+	ScalarType GetElementComponent(const uint ElementIndex, const uint ComponentIndex)
+	{
+		return M[ElementIndex][ComponentIndex];
+	}
+
+	CALL_SITE_DEBUGLOC
+	void SetElementComponent(const uint ElementIndex, const uint ComponentIndex, ScalarType Component)
+	{
+		M[ElementIndex][ComponentIndex] = Component;
+	}
+}; // struct FRegisters_AoS_Matrix
+
+/** Stores an array of vector<> as a structure of array in registers using a matrix<>. Requires InElementCount a pair >= 2 */
+template<typename ScalarType, uint InVectorSize, uint InElementCount>
+struct FRegisters_SoA_Matrix
 {
+	/** Size of each element. */
 	static const uint VectorSize = InVectorSize;
-	static const uint SimdSize = InSimdSizeX * InSimdSizeY;
-	static const uint SimdSizeX = InSimdSizeX;
-	static const uint SimdSizeY = InSimdSizeY;
+	
+	/** Number of elememnts */
+	static const uint ElementCount = InElementCount;
+	
+	/** Size and number of register rows */
+	static const uint RegisterRowSize = ElementCount;
+	static const uint RegisterRowCount = VectorSize;
 
-	vector<ScalarType, VectorSize> RegisterArray[SimdSize];
+	matrix<ScalarType, VectorSize, ElementCount> M;
+	
+	
+	// ------------- register rows
+	
+	CALL_SITE_DEBUGLOC
+	vector<ScalarType, RegisterRowSize> GetRegisterRow(const uint RegisterRowIndex)
+	{
+		return M[RegisterRowIndex];
+	}
+	
+	CALL_SITE_DEBUGLOC
+	void SetRegisterRow(const uint RegisterRowIndex, vector<ScalarType, RegisterRowSize> RegisterRow)
+	{
+		M[RegisterRowIndex] = RegisterRow;
+	}
+	
 
+	// ------------- elements
 	
-	// Access a component of the vector
 	CALL_SITE_DEBUGLOC
-	TLaneVector2D<ScalarType, 1, SimdSizeX, SimdSizeY> operator [](const uint ComponentIndex)
+	vector<ScalarType, VectorSize> GetElement(const uint ElementIndex)
 	{
-		TLaneVector2D<ScalarType, 1, SimdSizeX, SimdSizeY> R;
+		vector<ScalarType, VectorSize> Element;
 		UNROLL
-		for (uint SimdIndex = 0; SimdIndex < SimdSize; SimdIndex++)
+		for (uint ComponentIndex = 0; ComponentIndex < VectorSize; ComponentIndex++)
 		{
-			R.SetElement(SimdIndex, GetElement(SimdIndex)[ComponentIndex]);
+			Element[ComponentIndex] = M[ComponentIndex][ElementIndex];
 		}
-		return R;
+		return Element;
 	}
 	
-	// Access an individual element of the vector.
 	CALL_SITE_DEBUGLOC
-	vector<ScalarType, VectorSize> GetElement(const uint SimdIndex)
+	void SetElement(const uint ElementIndex, vector<ScalarType, VectorSize> Element)
 	{
-		return RegisterArray[SimdIndex];
+		UNROLL
+		for (uint ComponentIndex = 0; ComponentIndex < VectorSize; ComponentIndex++)
+		{
+			M[ComponentIndex][ElementIndex] = Element[ComponentIndex];
+		}
 	}
+
+	vector<ScalarType, VectorSize> WaveBroadcastElement(const FWaveBroadcastSettings BroadcastSettings, const uint ElementIndex)
+	{
+		vector<ScalarType, VectorSize> Element;
+		UNROLL
+		for (uint ComponentIndex = 0; ComponentIndex < VectorSize; ComponentIndex++)
+		{
+			Element[ComponentIndex] = WaveBroadcast(BroadcastSettings, M[ComponentIndex])[ElementIndex];
+		}
+		return Element;
+	}
+	
 	
+	// ------------- elements' components
+
+	CALL_SITE_DEBUGLOC
+	ScalarType GetElementComponent(const uint ElementIndex, const uint ComponentIndex)
+	{
+		return M[ComponentIndex][ElementIndex];
+	}
+
 	CALL_SITE_DEBUGLOC
-	void SetElement(const uint SimdIndex, vector<ScalarType, VectorSize> Element)
+	void SetElementComponent(const uint ElementIndex, const uint ComponentIndex, ScalarType Component)
 	{
-		RegisterArray[SimdIndex] = Element;
+		M[ComponentIndex][ElementIndex] = Component;
 	}
+}; // struct FRegisters_SoA_Matrix
+
+/** Stores an array of vector<> as a structure of array in registers using a matrix<>. Requires InElementCount a pair >= 2 */
+template<typename ScalarType, uint InVectorSize, uint InElementCount>
+struct FRegisters_AoS_PairArray
+{
+	/** Size of each element. */
+	static const uint VectorSize = InVectorSize;
+	
+	/** Number of elememnts */
+	static const uint ElementCount = InElementCount;
+	
+	/** Size and number of register rows */
+	static const uint RegisterRowSize = 2;
+	static const uint RegisterRowCount = (ElementCount * VectorSize) / 2;
+
+	vector<ScalarType, RegisterRowSize> M[RegisterRowCount];
+	
+	
+	// ------------- register rows
 	
 	CALL_SITE_DEBUGLOC
-	void SetAllElements(vector<ScalarType, VectorSize> Element)
+	vector<ScalarType, RegisterRowSize> GetRegisterRow(const uint RegisterRowIndex)
 	{
-		UNROLL
-		for (uint SimdIndex = 0; SimdIndex < SimdSize; SimdIndex++)
-		{
-			RegisterArray[SimdIndex] = Element;
-		}
+		return M[RegisterRowIndex];
 	}
 	
-	// Set component of the vector
 	CALL_SITE_DEBUGLOC
-	void SetComponent(uint ComponentIndex, TLaneVector2D<ScalarType, 1, SimdSizeX, SimdSizeY> Comp)
+	void SetRegisterRow(const uint RegisterRowIndex, vector<ScalarType, RegisterRowSize> RegisterRow)
 	{
-		UNROLL
-		for (uint SimdIndex = 0; SimdIndex < SimdSize; SimdIndex++)
-		{
-			RegisterArray[SimdIndex][ComponentIndex] = Comp.RegisterArray[SimdIndex];
-		}
+		M[RegisterRowIndex] = RegisterRow;
 	}
+	
 
+	// ------------- elements
 	
-	// ------------- binary operator +
 	CALL_SITE_DEBUGLOC
-	TLaneVector2D<ScalarType, VectorSize, SimdSizeX, SimdSizeY> operator + (TLaneVector2D<ScalarType, VectorSize, SimdSizeX, SimdSizeY> B)
+	vector<ScalarType, VectorSize> GetElement(const uint ElementIndex)
 	{
-		TLaneVector2D<ScalarType, VectorSize, SimdSizeX, SimdSizeY> R;
+		vector<ScalarType, VectorSize> Element;
 		UNROLL
-		for (uint SimdIndex = 0; SimdIndex < SimdSize; SimdIndex++)
+		for (uint ComponentIndex = 0; ComponentIndex < VectorSize; ComponentIndex++)
 		{
-			R.SetElement(SimdIndex, GetElement(SimdIndex) + B.GetElement(SimdIndex));
+			Element[ComponentIndex] = GetElementComponent(ElementIndex, ComponentIndex);
 		}
-		return R;
+		return Element;
 	}
 	
 	CALL_SITE_DEBUGLOC
-	TLaneVector2D<ScalarType, VectorSize, SimdSizeX, SimdSizeY> operator + (vector<ScalarType, VectorSize> B)
+	void SetElement(const uint ElementIndex, vector<ScalarType, VectorSize> Element)
 	{
-		TLaneVector2D<ScalarType, VectorSize, SimdSizeX, SimdSizeY> R;
 		UNROLL
-		for (uint SimdIndex = 0; SimdIndex < SimdSize; SimdIndex++)
+		for (uint ComponentIndex = 0; ComponentIndex < VectorSize; ComponentIndex++)
 		{
-			R.SetElement(SimdIndex, GetElement(SimdIndex) + B);
+			SetElementComponent(ElementIndex, ComponentIndex, Element[ComponentIndex]);
 		}
-		return R;
 	}
 
-	
-	// ------------- unary operator -
-	CALL_SITE_DEBUGLOC
-	TLaneVector2D<ScalarType, VectorSize, SimdSizeX, SimdSizeY> operator - ()
+	vector<ScalarType, VectorSize> WaveBroadcastElement(const FWaveBroadcastSettings BroadcastSettings, const uint ElementIndex)
 	{
-		TLaneVector2D<ScalarType, VectorSize, SimdSizeX, SimdSizeY> R;
-		UNROLL
-		for (uint SimdIndex = 0; SimdIndex < SimdSize; SimdIndex++)
+		vector<ScalarType, VectorSize> Element;
+		if (sizeof(ScalarType) == 4)
 		{
-			R.SetElement(SimdIndex, -GetElement(SimdIndex));
+			UNROLL
+			for (uint ComponentIndex = 0; ComponentIndex < VectorSize; ComponentIndex++)
+			{
+				const uint ArrayIndex = ComponentIndex + VectorSize * ElementIndex;
+				Element[ComponentIndex] = WaveBroadcast(BroadcastSettings, M[ArrayIndex / 2u][ArrayIndex % 2u]);
+			}
 		}
-		return R;
-	}
-	
-	// ------------- binary operator -
-	CALL_SITE_DEBUGLOC
-	TLaneVector2D<ScalarType, VectorSize, SimdSizeX, SimdSizeY> operator - (TLaneVector2D<ScalarType, VectorSize, SimdSizeX, SimdSizeY> B)
-	{
-		TLaneVector2D<ScalarType, VectorSize, SimdSizeX, SimdSizeY> R;
-		UNROLL
-		for (uint SimdIndex = 0; SimdIndex < SimdSize; SimdIndex++)
+		else
 		{
-			R.SetElement(SimdIndex, GetElement(SimdIndex) - B.GetElement(SimdIndex));
+			const uint FirstArrayIndex = VectorSize * ElementIndex;
+			
+			vector<ScalarType, RegisterRowSize> BroadcastedArray[(VectorSize + 1u) / 2u];
+			UNROLL
+			for (uint BroadcastIndex = 0; BroadcastIndex < ((VectorSize + 1u) / 2u); BroadcastIndex++)
+			{
+				BroadcastedArray[BroadcastIndex] = WaveBroadcast(BroadcastSettings, M[FirstArrayIndex / 2u + BroadcastIndex]);
+			}
+			
+			UNROLL
+			for (uint ComponentIndex = 0; ComponentIndex < VectorSize; ComponentIndex++)
+			{
+				const uint BroadcastedArrayIndex = ComponentIndex + VectorSize * ElementIndex - (FirstArrayIndex / 2u) * 2u;
+				Element[ComponentIndex] = BroadcastedArray[BroadcastedArrayIndex / 2u][BroadcastedArrayIndex % 2u];
+			}
 		}
-		return R;
+		return Element;
 	}
 	
+	
+	// ------------- elements' components
+
 	CALL_SITE_DEBUGLOC
-	TLaneVector2D<ScalarType, VectorSize, SimdSizeX, SimdSizeY> operator - (vector<ScalarType, VectorSize> B)
+	ScalarType GetElementComponent(const uint ElementIndex, const uint ComponentIndex)
 	{
-		TLaneVector2D<ScalarType, VectorSize, SimdSizeX, SimdSizeY> R;
-		UNROLL
-		for (uint SimdIndex = 0; SimdIndex < SimdSize; SimdIndex++)
-		{
-			R.SetElement(SimdIndex, GetElement(SimdIndex) - B);
-		}
-		return R;
+		const uint ArrayIndex = ComponentIndex + VectorSize * ElementIndex;
+		return M[ArrayIndex / 2u][ArrayIndex % 2u];
 	}
 
-	
-	// ------------- binary operator *
 	CALL_SITE_DEBUGLOC
-	TLaneVector2D<ScalarType, VectorSize, SimdSizeX, SimdSizeY> operator * (TLaneVector2D<ScalarType, VectorSize, SimdSizeX, SimdSizeY> B)
+	void SetElementComponent(const uint ElementIndex, const uint ComponentIndex, ScalarType Component)
 	{
-		TLaneVector2D<ScalarType, VectorSize, SimdSizeX, SimdSizeY> R;
-		UNROLL
-		for (uint SimdIndex = 0; SimdIndex < SimdSize; SimdIndex++)
-		{
-			R.SetElement(SimdIndex, GetElement(SimdIndex) * B.GetElement(SimdIndex));
-		}
-		return R;
+		const uint ArrayIndex = ComponentIndex + VectorSize * ElementIndex;
+		M[ArrayIndex / 2u][ArrayIndex % 2u] = Component;
 	}
+}; // struct FRegisters_AoS_PairArray
+
+/** Stores an array of vector<> as a structure of array in registers using a matrix<>. Requires InElementCount a pair >= 2 */
+template<typename ScalarType, uint InVectorSize, uint InElementCount>
+struct FRegisters_SoA_PairArray
+{
+	/** Size of each element. */
+	static const uint VectorSize = InVectorSize;
+	
+	/** Number of elememnts */
+	static const uint ElementCount = InElementCount;
+	
+	/** Size and number of register rows */
+	static const uint RegisterRowSize = 2;
+	static const uint RegisterRowCount = (ElementCount * VectorSize) / 2;
 
+	vector<ScalarType, RegisterRowSize> M[RegisterRowCount];
+	
+	
+	// ------------- register rows
+	
 	CALL_SITE_DEBUGLOC
-	TLaneVector2D<ScalarType, VectorSize, SimdSizeX, SimdSizeY> operator * (vector<ScalarType, VectorSize> B)
+	vector<ScalarType, RegisterRowSize> GetRegisterRow(const uint RegisterRowIndex)
 	{
-		TLaneVector2D<ScalarType, VectorSize, SimdSizeX, SimdSizeY> R;
-		UNROLL
-		for (uint SimdIndex = 0; SimdIndex < SimdSize; SimdIndex++)
-		{
-			R.SetElement(SimdIndex, GetElement(SimdIndex) * B);
-		}
-		return R;
+		return M[RegisterRowIndex];
 	}
 	
-	// ------------- comparison operators
 	CALL_SITE_DEBUGLOC
-	TLaneVector2D<bool, VectorSize, SimdSizeX, SimdSizeY> operator > (TLaneVector2D<ScalarType, VectorSize, SimdSizeX, SimdSizeY> B)
+	void SetRegisterRow(const uint RegisterRowIndex, vector<ScalarType, RegisterRowSize> RegisterRow)
 	{
-		TLaneVector2D<bool, VectorSize, SimdSizeX, SimdSizeY> R;
-		UNROLL
-		for (uint SimdIndex = 0; SimdIndex < SimdSize; SimdIndex++)
-		{
-			R.SetElement(SimdIndex, GetElement(SimdIndex) > B.GetElement(SimdIndex));
-		}
-		return R;
+		M[RegisterRowIndex] = RegisterRow;
 	}
 	
+
+	// ------------- elements
+	
 	CALL_SITE_DEBUGLOC
-	TLaneVector2D<bool, VectorSize, SimdSizeX, SimdSizeY> operator < (TLaneVector2D<ScalarType, VectorSize, SimdSizeX, SimdSizeY> B)
+	vector<ScalarType, VectorSize> GetElement(const uint ElementIndex)
 	{
-		TLaneVector2D<bool, VectorSize, SimdSizeX, SimdSizeY> R;
+		vector<ScalarType, VectorSize> Element;
 		UNROLL
-		for (uint SimdIndex = 0; SimdIndex < SimdSize; SimdIndex++)
+		for (uint ComponentIndex = 0; ComponentIndex < VectorSize; ComponentIndex++)
 		{
-			R.SetElement(SimdIndex, GetElement(SimdIndex) < B.GetElement(SimdIndex));
+			Element[ComponentIndex] = GetElementComponent(ElementIndex, ComponentIndex);
 		}
-		return R;
+		return Element;
 	}
-
-
-	// ------------- constructors
+	
 	CALL_SITE_DEBUGLOC
-	static TLaneVector2D<ScalarType, VectorSize, SimdSizeX, SimdSizeY> Vectorize(
-		TLaneVector2D<ScalarType, 1, SimdSizeX, SimdSizeY> A)
+	void SetElement(const uint ElementIndex, vector<ScalarType, VectorSize> Element)
 	{
-		TLaneVector2D<ScalarType, VectorSize, SimdSizeX, SimdSizeY> R;
 		UNROLL
-		for (uint SimdIndex = 0; SimdIndex < SimdSize; SimdIndex++)
+		for (uint ComponentIndex = 0; ComponentIndex < VectorSize; ComponentIndex++)
 		{
-			R.SetElement(SimdIndex, A.GetElement(SimdIndex));
+			SetElementComponent(ElementIndex, ComponentIndex, Element[ComponentIndex]);
 		}
-		return R;
 	}
-	
-	CALL_SITE_DEBUGLOC
-	static TLaneVector2D<ScalarType, VectorSize, SimdSizeX, SimdSizeY> Const(
-		vector<ScalarType, VectorSize> A)
+
+	vector<ScalarType, VectorSize> WaveBroadcastElement(const FWaveBroadcastSettings BroadcastSettings, const uint ElementIndex)
 	{
-		TLaneVector2D<ScalarType, VectorSize, SimdSizeX, SimdSizeY> R;
+		vector<ScalarType, VectorSize> Element;
 		UNROLL
-		for (uint SimdIndex = 0; SimdIndex < SimdSize; SimdIndex++)
+		for (uint ComponentIndex = 0; ComponentIndex < VectorSize; ComponentIndex++)
 		{
-			R.SetElement(SimdIndex, A);
+			const uint ArrayIndex = ElementIndex + VectorSize * ComponentIndex;
+			Element[ComponentIndex] = WaveBroadcast(BroadcastSettings, M[ArrayIndex / 2u])[ArrayIndex % 2u];
 		}
-		return R;
+		return Element;
 	}
 	
+	
+	// ------------- elements' components
+
 	CALL_SITE_DEBUGLOC
-	template<typename SourceScalarType>
-	static TLaneVector2D<ScalarType, VectorSize, SimdSizeX, SimdSizeY> CastFrom(
-		TLaneVector2D<SourceScalarType, 1, SimdSizeX, SimdSizeY> A)
+	ScalarType GetElementComponent(const uint ElementIndex, const uint ComponentIndex)
 	{
-		TLaneVector2D<ScalarType, VectorSize, SimdSizeX, SimdSizeY> R;
-		UNROLL
-		for (uint SimdIndex = 0; SimdIndex < SimdSize; SimdIndex++)
-		{
-			R.SetElement(SimdIndex, vector<ScalarType, VectorSize>(A.GetElement(SimdIndex)));
-		}
-		return R;
+		const uint ArrayIndex = ElementIndex + VectorSize * ComponentIndex;
+		return M[ArrayIndex / 2u][ArrayIndex % 2u];
 	}
-	
+
 	CALL_SITE_DEBUGLOC
-	static TLaneVector2D<ScalarType, VectorSize, SimdSizeX, SimdSizeY> ForceRegisterPacking(
-		TLaneVector2D<ScalarType, VectorSize, SimdSizeX, SimdSizeY> A)
+	void SetElementComponent(const uint ElementIndex, const uint ComponentIndex, ScalarType Component)
 	{
-		return A;
+		const uint ArrayIndex = ElementIndex + VectorSize * ComponentIndex;
+		M[ArrayIndex / 2u][ArrayIndex % 2u] = Component;
 	}
-	
-}; // TLaneVector2D
+}; // struct FRegisters_SoA_PairArray
+
 
-/** Work exactly like a TLaneVector2D, except values are packed vector<ScalarType, 2> to better hint the shader compiler to use 16bit packed instructions.  */
+//------------------------------------------------------- TENSORS
+
+/** Packs a vector<> in 2 dimension. */
 template<typename ScalarType, uint InVectorSize, uint InSimdSizeX, uint InSimdSizeY>
-struct TLanePackedVector2D
+struct TLaneVector2D
 {
+	/** Size of each element. */
 	static const uint VectorSize = InVectorSize;
+	
+	/** Number of elememnts */
 	static const uint SimdSize = InSimdSizeX * InSimdSizeY;
 	static const uint SimdSizeX = InSimdSizeX;
 	static const uint SimdSizeY = InSimdSizeY;
-	static const uint ScalarPairSize = (VectorSize * SimdSize) / 2u;
+	
+	/** Size and number of register rows */
+	static const uint RegisterRowSize  = TENSOR_REGISTER_LAYOUT<ScalarType, InVectorSize, InSimdSizeX * InSimdSizeY>::RegisterRowSize;
+	static const uint RegisterRowCount = TENSOR_REGISTER_LAYOUT<ScalarType, InVectorSize, InSimdSizeX * InSimdSizeY>::RegisterRowCount;
 
-	vector<ScalarType, 2u> RegisterArray[VectorSize][SimdSize / 2u];
+	TENSOR_REGISTER_LAYOUT<ScalarType, InVectorSize, InSimdSizeX * InSimdSizeY> Registers;
 
 	
-	// Access a component of the vector
+	// ------------- getting and setting elements and components
+
+	// Access and set individual element of the vector.
 	CALL_SITE_DEBUGLOC
-	TLanePackedVector2D<ScalarType, 1, SimdSizeX, SimdSizeY> operator [](const uint ComponentIndex)
+	vector<ScalarType, VectorSize> GetElement(const uint ElementIndex)
 	{
-		TLanePackedVector2D<ScalarType, 1, SimdSizeX, SimdSizeY> R;
-		UNROLL
-		for (uint DestScalarPairIndex = 0; DestScalarPairIndex < (SimdSizeX * SimdSizeY) / 2u; DestScalarPairIndex++)
-		{
-			R.SetScalarPair(DestScalarPairIndex, GetScalarPair(DestScalarPairIndex * VectorSize + ComponentIndex));
-		}
-		return R;
+		return Registers.GetElement(ElementIndex);
 	}
 	
-	// Access an individual element of the vector.
 	CALL_SITE_DEBUGLOC
-	vector<ScalarType, VectorSize> GetElement(const uint SimdIndex)
+	void SetElement(const uint ElementIndex, vector<ScalarType, VectorSize> Element)
 	{
-		vector<ScalarType, VectorSize> Element;
-		UNROLL
-		for (uint ComponentIndex = 0; ComponentIndex < VectorSize; ComponentIndex++)
-		{
-			Element[ComponentIndex] = RegisterArray[ComponentIndex][SimdIndex / 2u][SimdIndex % 2u];
-		}
-		return Element;
+		Registers.SetElement(ElementIndex, Element);
 	}
 	
 	CALL_SITE_DEBUGLOC
-	void SetElement(const uint SimdIndex, vector<ScalarType, VectorSize> Element)
+	void SetAllElements(vector<ScalarType, VectorSize> Element)
 	{
 		UNROLL
-		for (uint ComponentIndex = 0; ComponentIndex < VectorSize; ComponentIndex++)
+		for (uint ElementIndex = 0; ElementIndex < SimdSize; ElementIndex++)
 		{
-			RegisterArray[ComponentIndex][SimdIndex / 2u][SimdIndex % 2u] = Element[ComponentIndex];
+			Registers.SetElement(ElementIndex, Element);
 		}
 	}
 	
+	
+	//  Access and set a component of the vector
 	CALL_SITE_DEBUGLOC
-	void SetAllElements(vector<ScalarType, VectorSize> Element)
+	TLaneVector2D<ScalarType, 1, SimdSizeX, SimdSizeY> GetComponent(const uint ComponentIndex)
 	{
+		TLaneVector2D<ScalarType, 1, SimdSizeX, SimdSizeY> R;
 		UNROLL
-		for (uint SimdIndex = 0; SimdIndex < SimdSize; SimdIndex++)
+		for (uint ElementIndex = 0; ElementIndex < SimdSize; ElementIndex++)
 		{
-			SetElement(SimdIndex, Element);
+			R.SetElement(ElementIndex, Registers.GetElementComponent(ElementIndex, ComponentIndex));
 		}
+		return R;
 	}
 	
-	// Access an individual scalair pair.
 	CALL_SITE_DEBUGLOC
-	vector<ScalarType, 2> GetScalarPair(const uint ScalarPairIndex)
+	void SetComponent(uint ComponentIndex, TLaneVector2D<ScalarType, 1, SimdSizeX, SimdSizeY> Comp)
 	{
-		return RegisterArray[ScalarPairIndex % VectorSize][ScalarPairIndex / VectorSize];
+		UNROLL
+		for (uint ElementIndex = 0; ElementIndex < SimdSize; ElementIndex++)
+		{
+			Registers.SetElementComponent(ElementIndex, ComponentIndex, Comp.GetElement(ElementIndex));
+		}
 	}
 	
 	CALL_SITE_DEBUGLOC
-	void SetScalarPair(const uint ScalarPairIndex, vector<ScalarType, 2> ScalarPair)
+	TLaneVector2D<ScalarType, 1, SimdSizeX, SimdSizeY> operator [](const uint ComponentIndex)
 	{
-		RegisterArray[ScalarPairIndex % VectorSize][ScalarPairIndex / VectorSize] = ScalarPair;
+		return GetComponent(ComponentIndex);
 	}
 	
-	// Set component of the vector
+
+	// ------------- constructors
+
+	/** Casts a scalar into a vector<>. */
 	CALL_SITE_DEBUGLOC
-	void SetComponent(uint ComponentIndex, TLanePackedVector2D<ScalarType, 1, SimdSizeX, SimdSizeY> Comp)
+	static TLaneVector2D<ScalarType, VectorSize, SimdSizeX, SimdSizeY> Vectorize(
+		TLaneVector2D<ScalarType, 1, SimdSizeX, SimdSizeY> A)
 	{
+		TLaneVector2D<ScalarType, VectorSize, SimdSizeX, SimdSizeY> R;
 		UNROLL
-		for (uint SrcPairIndex = 0; SrcPairIndex < (SimdSizeX * SimdSizeY) / 2u; SrcPairIndex++)
+		for (uint ElementIndex = 0; ElementIndex < SimdSize; ElementIndex++)
 		{
-			SetScalarPair(SrcPairIndex * VectorSize + ComponentIndex, Comp.GetScalarPair(SrcPairIndex));
+			R.SetElement(ElementIndex, A.GetElement(ElementIndex));
 		}
+		return R;
 	}
 	
-	// ------------- binary operator +
+	/** Initialize all elements with a single same vector<>. */
 	CALL_SITE_DEBUGLOC
-	TLanePackedVector2D<ScalarType, VectorSize, SimdSizeX, SimdSizeY> operator + (TLanePackedVector2D<ScalarType, VectorSize, SimdSizeX, SimdSizeY> B)
+	static TLaneVector2D<ScalarType, VectorSize, SimdSizeX, SimdSizeY> Const(
+		vector<ScalarType, VectorSize> A)
 	{
-		TLanePackedVector2D<ScalarType, VectorSize, SimdSizeX, SimdSizeY> R;
+		TLaneVector2D<ScalarType, VectorSize, SimdSizeX, SimdSizeY> R;
 		UNROLL
-		for (uint ScalarPairIndex = 0; ScalarPairIndex < ScalarPairSize; ScalarPairIndex++)
+		for (uint ElementIndex = 0; ElementIndex < SimdSize; ElementIndex++)
 		{
-			R.SetScalarPair(ScalarPairIndex, GetScalarPair(ScalarPairIndex) + B.GetScalarPair(ScalarPairIndex));
+			R.SetElement(ElementIndex, A);
 		}
 		return R;
 	}
 	
+	/** Casts vector<A> to vector<B>. */
 	CALL_SITE_DEBUGLOC
-	TLanePackedVector2D<ScalarType, VectorSize, SimdSizeX, SimdSizeY> operator + (vector<ScalarType, VectorSize> B)
+	template<typename SourceScalarType>
+	static TLaneVector2D<ScalarType, VectorSize, SimdSizeX, SimdSizeY> CastFrom(
+		TLaneVector2D<SourceScalarType, 1, SimdSizeX, SimdSizeY> A)
 	{
-		TLanePackedVector2D<ScalarType, VectorSize, SimdSizeX, SimdSizeY> R;
+		TLaneVector2D<ScalarType, VectorSize, SimdSizeX, SimdSizeY> R;
 		UNROLL
-		for (uint ScalarPairIndex = 0; ScalarPairIndex < ScalarPairSize; ScalarPairIndex++)
+		for (uint ElementIndex = 0; ElementIndex < SimdSize; ElementIndex++)
 		{
-			R.SetScalarPair(ScalarPairIndex, GetScalarPair(ScalarPairIndex) + B[ScalarPairIndex % VectorSize]);
+			R.SetElement(ElementIndex, vector<ScalarType, VectorSize>(A.GetElement(ElementIndex)));
 		}
 		return R;
 	}
-
 	
-	// ------------- unary operator -
+	/** Force tight register packing. */
 	CALL_SITE_DEBUGLOC
-	TLanePackedVector2D<ScalarType, VectorSize, SimdSizeX, SimdSizeY> operator - ()
+	void TightenRegisters()
+	#if TENSOR_REGISTER_TIGHTENING
 	{
-		TLanePackedVector2D<ScalarType, VectorSize, SimdSizeX, SimdSizeY> R;
 		UNROLL
-		for (uint ScalarPairIndex = 0; ScalarPairIndex < ScalarPairSize; ScalarPairIndex++)
+		for (uint RegisterRowIndex = 0; RegisterRowIndex < RegisterRowCount; RegisterRowIndex++)
 		{
-			R.SetScalarPair(ScalarPairIndex, -GetScalarPair(ScalarPairIndex));
+			Registers.SetRegisterRow(RegisterRowIndex, PackRegistersTightly(Registers.GetRegisterRow(RegisterRowIndex)));
 		}
-		return R;
 	}
-	
-	// ------------- binary operator -
-	CALL_SITE_DEBUGLOC
-	TLanePackedVector2D<ScalarType, VectorSize, SimdSizeX, SimdSizeY> operator - (TLanePackedVector2D<ScalarType, VectorSize, SimdSizeX, SimdSizeY> B)
+	#else
 	{
-		TLanePackedVector2D<ScalarType, VectorSize, SimdSizeX, SimdSizeY> R;
-		UNROLL
-		for (uint ScalarPairIndex = 0; ScalarPairIndex < ScalarPairSize; ScalarPairIndex++)
-		{
-			R.SetScalarPair(ScalarPairIndex, GetScalarPair(ScalarPairIndex) - B.GetScalarPair(ScalarPairIndex));
-		}
-		return R;
+		// NOP
 	}
+	#endif
+	
 	
+	// ------------- binary operator +
 	CALL_SITE_DEBUGLOC
-	TLanePackedVector2D<ScalarType, VectorSize, SimdSizeX, SimdSizeY> operator - (vector<ScalarType, VectorSize> B)
+	TLaneVector2D<ScalarType, VectorSize, SimdSizeX, SimdSizeY> operator + (TLaneVector2D<ScalarType, VectorSize, SimdSizeX, SimdSizeY> B)
 	{
-		TLanePackedVector2D<ScalarType, VectorSize, SimdSizeX, SimdSizeY> R;
+		TLaneVector2D<ScalarType, VectorSize, SimdSizeX, SimdSizeY> R;
 		UNROLL
-		for (uint ScalarPairIndex = 0; ScalarPairIndex < ScalarPairSize; ScalarPairIndex++)
+		for (uint RegisterRowIndex = 0; RegisterRowIndex < RegisterRowCount; RegisterRowIndex++)
 		{
-			R.SetScalarPair(ScalarPairIndex, GetScalarPair(ScalarPairIndex) - B[ScalarPairIndex % VectorSize]);
+			R.Registers.SetRegisterRow(RegisterRowIndex, Registers.GetRegisterRow(RegisterRowIndex) + B.Registers.GetRegisterRow(RegisterRowIndex));
 		}
 		return R;
 	}
-
 	
-	// ------------- binary operator *
 	CALL_SITE_DEBUGLOC
-	TLanePackedVector2D<ScalarType, VectorSize, SimdSizeX, SimdSizeY> operator * (TLanePackedVector2D<ScalarType, VectorSize, SimdSizeX, SimdSizeY> B)
+	TLaneVector2D<ScalarType, VectorSize, SimdSizeX, SimdSizeY> operator + (vector<ScalarType, VectorSize> B)
 	{
-		TLanePackedVector2D<ScalarType, VectorSize, SimdSizeX, SimdSizeY> R;
+		TLaneVector2D<ScalarType, VectorSize, SimdSizeX, SimdSizeY> BV;
+		BV.SetAllElements(B);
+
+		TLaneVector2D<ScalarType, VectorSize, SimdSizeX, SimdSizeY> R;
 		UNROLL
-		for (uint ScalarPairIndex = 0; ScalarPairIndex < ScalarPairSize; ScalarPairIndex++)
+		for (uint RegisterRowIndex = 0; RegisterRowIndex < RegisterRowCount; RegisterRowIndex++)
 		{
-			R.SetScalarPair(ScalarPairIndex, GetScalarPair(ScalarPairIndex) * B.GetScalarPair(ScalarPairIndex));
+			R.Registers.SetRegisterRow(RegisterRowIndex, Registers.GetRegisterRow(RegisterRowIndex) + BV.Registers.GetRegisterRow(RegisterRowIndex));
 		}
 		return R;
 	}
 
+	
+	// ------------- unary operator -
 	CALL_SITE_DEBUGLOC
-	TLanePackedVector2D<ScalarType, VectorSize, SimdSizeX, SimdSizeY> operator * (vector<ScalarType, VectorSize> B)
+	TLaneVector2D<ScalarType, VectorSize, SimdSizeX, SimdSizeY> operator - ()
 	{
-		TLanePackedVector2D<ScalarType, VectorSize, SimdSizeX, SimdSizeY> R;
+		TLaneVector2D<ScalarType, VectorSize, SimdSizeX, SimdSizeY> R;
 		UNROLL
-		for (uint ScalarPairIndex = 0; ScalarPairIndex < ScalarPairSize; ScalarPairIndex++)
+		for (uint RegisterRowIndex = 0; RegisterRowIndex < RegisterRowCount; RegisterRowIndex++)
 		{
-			R.SetScalarPair(ScalarPairIndex, GetScalarPair(ScalarPairIndex) * B[ScalarPairIndex % VectorSize]);
+			R.Registers.SetRegisterRow(RegisterRowIndex, -Registers.GetRegisterRow(RegisterRowIndex));
 		}
 		return R;
 	}
 	
-	// ------------- comparison operators
+	// ------------- binary operator -
 	CALL_SITE_DEBUGLOC
-	TLanePackedVector2D<bool, VectorSize, SimdSizeX, SimdSizeY> operator > (TLanePackedVector2D<ScalarType, VectorSize, SimdSizeX, SimdSizeY> B)
+	TLaneVector2D<ScalarType, VectorSize, SimdSizeX, SimdSizeY> operator - (TLaneVector2D<ScalarType, VectorSize, SimdSizeX, SimdSizeY> B)
 	{
-		TLanePackedVector2D<bool, VectorSize, SimdSizeX, SimdSizeY> R;
+		TLaneVector2D<ScalarType, VectorSize, SimdSizeX, SimdSizeY> R;
 		UNROLL
-		for (uint ScalarPairIndex = 0; ScalarPairIndex < ScalarPairSize; ScalarPairIndex++)
+		for (uint RegisterRowIndex = 0; RegisterRowIndex < RegisterRowCount; RegisterRowIndex++)
 		{
-			R.SetScalarPair(ScalarPairIndex, GetScalarPair(ScalarPairIndex) > B.GetScalarPair(ScalarPairIndex));
+			R.Registers.SetRegisterRow(RegisterRowIndex, Registers.GetRegisterRow(RegisterRowIndex) - B.Registers.GetRegisterRow(RegisterRowIndex));
 		}
 		return R;
 	}
 	
 	CALL_SITE_DEBUGLOC
-	TLanePackedVector2D<bool, VectorSize, SimdSizeX, SimdSizeY> operator < (TLanePackedVector2D<ScalarType, VectorSize, SimdSizeX, SimdSizeY> B)
+	TLaneVector2D<ScalarType, VectorSize, SimdSizeX, SimdSizeY> operator - (vector<ScalarType, VectorSize> B)
 	{
-		TLanePackedVector2D<bool, VectorSize, SimdSizeX, SimdSizeY> R;
+		TLaneVector2D<ScalarType, VectorSize, SimdSizeX, SimdSizeY> BV;
+		BV.SetAllElements(B);
+
+		TLaneVector2D<ScalarType, VectorSize, SimdSizeX, SimdSizeY> R;
 		UNROLL
-		for (uint ScalarPairIndex = 0; ScalarPairIndex < ScalarPairSize; ScalarPairIndex++)
+		for (uint RegisterRowIndex = 0; RegisterRowIndex < RegisterRowCount; RegisterRowIndex++)
 		{
-			R.SetScalarPair(ScalarPairIndex, GetScalarPair(ScalarPairIndex) < B.GetScalarPair(ScalarPairIndex));
+			R.Registers.SetRegisterRow(RegisterRowIndex, Registers.GetRegisterRow(RegisterRowIndex) - BV.Registers.GetRegisterRow(RegisterRowIndex));
 		}
 		return R;
 	}
 
-	// ------------- constructors
+	
+	// ------------- binary operator *
 	CALL_SITE_DEBUGLOC
-	static TLanePackedVector2D<ScalarType, VectorSize, SimdSizeX, SimdSizeY> Vectorize(
-		TLanePackedVector2D<ScalarType, 1, SimdSizeX, SimdSizeY> A)
+	TLaneVector2D<ScalarType, VectorSize, SimdSizeX, SimdSizeY> operator * (TLaneVector2D<ScalarType, VectorSize, SimdSizeX, SimdSizeY> B)
 	{
-		TLanePackedVector2D<ScalarType, VectorSize, SimdSizeX, SimdSizeY> R;
+		TLaneVector2D<ScalarType, VectorSize, SimdSizeX, SimdSizeY> R;
 		UNROLL
-		for (uint SimdIndex = 0; SimdIndex < SimdSize; SimdIndex++)
+		for (uint RegisterRowIndex = 0; RegisterRowIndex < RegisterRowCount; RegisterRowIndex++)
 		{
-			R.SetElement(SimdIndex, A.GetElement(SimdIndex));
+			R.Registers.SetRegisterRow(RegisterRowIndex, Registers.GetRegisterRow(RegisterRowIndex) * B.Registers.GetRegisterRow(RegisterRowIndex));
 		}
 		return R;
 	}
-	
-	CALL_SITE_DEBUGLOC
-	static TLanePackedVector2D<ScalarType, VectorSize, SimdSizeX, SimdSizeY> Const(
-		vector<ScalarType, VectorSize> A)
-	{
-		TLanePackedVector2D<ScalarType, VectorSize, SimdSizeX, SimdSizeY> R;
-		R.SetAllElements(A);
-		return R;
-	}
-	
+
 	CALL_SITE_DEBUGLOC
-	template<typename SourceScalarType>
-	static TLanePackedVector2D<ScalarType, VectorSize, SimdSizeX, SimdSizeY> CastFrom(
-		TLanePackedVector2D<SourceScalarType, VectorSize, SimdSizeX, SimdSizeY> A)
+	TLaneVector2D<ScalarType, VectorSize, SimdSizeX, SimdSizeY> operator * (vector<ScalarType, VectorSize> B)
 	{
-		TLanePackedVector2D<ScalarType, VectorSize, SimdSizeX, SimdSizeY> R;
+		TLaneVector2D<ScalarType, VectorSize, SimdSizeX, SimdSizeY> BV;
+		BV.SetAllElements(B);
+
+		TLaneVector2D<ScalarType, VectorSize, SimdSizeX, SimdSizeY> R;
 		UNROLL
-		for (uint ScalarPairIndex = 0; ScalarPairIndex < ScalarPairSize; ScalarPairIndex++)
+		for (uint RegisterRowIndex = 0; RegisterRowIndex < RegisterRowCount; RegisterRowIndex++)
 		{
-			R.SetScalarPair(ScalarPairIndex, vector<ScalarType, 2>(A.GetScalarPair(ScalarPairIndex)));
+			R.Registers.SetRegisterRow(RegisterRowIndex, Registers.GetRegisterRow(RegisterRowIndex) * BV.Registers.GetRegisterRow(RegisterRowIndex));
 		}
 		return R;
 	}
 	
+	// ------------- comparison operators
 	CALL_SITE_DEBUGLOC
-	static TLanePackedVector2D<ScalarType, VectorSize, SimdSizeX, SimdSizeY> ForceRegisterPacking(
-		TLanePackedVector2D<ScalarType, VectorSize, SimdSizeX, SimdSizeY> A)
+	TLaneVector2D<bool, VectorSize, SimdSizeX, SimdSizeY> operator > (TLaneVector2D<ScalarType, VectorSize, SimdSizeX, SimdSizeY> B)
 	{
-		TLanePackedVector2D<ScalarType, VectorSize, SimdSizeX, SimdSizeY> R;
+		TLaneVector2D<bool, VectorSize, SimdSizeX, SimdSizeY> R;
 		UNROLL
-		for (uint ScalarPairIndex = 0; ScalarPairIndex < ScalarPairSize; ScalarPairIndex++)
+		for (uint RegisterRowIndex = 0; RegisterRowIndex < RegisterRowCount; RegisterRowIndex++)
 		{
-			vector<ScalarType, 2> Pair = A.GetScalarPair(ScalarPairIndex);
-			R.SetScalarPair(ScalarPairIndex, v_pack_b32_b16(Pair.x, Pair.y));
+			R.Registers.SetRegisterRow(RegisterRowIndex, Registers.GetRegisterRow(RegisterRowIndex) > B.Registers.GetRegisterRow(RegisterRowIndex));
 		}
 		return R;
 	}
 	
 	CALL_SITE_DEBUGLOC
-	static TLanePackedVector2D<ScalarType, VectorSize, SimdSizeX, SimdSizeY> ChangeRegisterOrder(
-		TLaneVector2D<ScalarType, VectorSize, SimdSizeX, SimdSizeY> A)
+	TLaneVector2D<bool, VectorSize, SimdSizeX, SimdSizeY> operator < (TLaneVector2D<ScalarType, VectorSize, SimdSizeX, SimdSizeY> B)
 	{
-		TLanePackedVector2D<ScalarType, VectorSize, SimdSizeX, SimdSizeY> R;
+		TLaneVector2D<bool, VectorSize, SimdSizeX, SimdSizeY> R;
 		UNROLL
-		for (uint SimdIndex = 0; SimdIndex < SimdSize; SimdIndex++)
+		for (uint RegisterRowIndex = 0; RegisterRowIndex < RegisterRowCount; RegisterRowIndex++)
 		{
-			R.SetElement(SimdIndex, A.GetElement(SimdIndex));
+			R.Registers.SetRegisterRow(RegisterRowIndex, Registers.GetRegisterRow(RegisterRowIndex) < B.Registers.GetRegisterRow(RegisterRowIndex));
 		}
 		return R;
 	}
-}; // struct TLanePackedVector2D
+
+
+}; // TLaneVector2D
 
 
 // ------------- 1 parameter
 
 #define TVECTOR_FUNCTION_1PARAMS(ReturnScalarType, FunctionName) \
-	CALL_SITE_DEBUGLOC \
-	template<typename ScalarType, uint VectorSize, uint SimdSize> \
-	TLaneVector<ReturnScalarType, VectorSize, SimdSize> FunctionName( \
-		TLaneVector<ScalarType, VectorSize, SimdSize> A) \
-	{ \
-		TLaneVector<ReturnScalarType, VectorSize, SimdSize> R; \
-		UNROLL \
-		for (uint SimdIndex = 0; SimdIndex < SimdSize; SimdIndex++) \
-		{ \
-			R.SetElement(SimdIndex, FunctionName(A.GetElement(SimdIndex))); \
-		} \
-		return R; \
-	} \
-	\
 	CALL_SITE_DEBUGLOC \
 	template<typename ScalarType, uint VectorSize, uint SimdSizeX, uint SimdSizeY> \
 	TLaneVector2D<ReturnScalarType, VectorSize, SimdSizeX, SimdSizeY> FunctionName( \
 		TLaneVector2D<ScalarType, VectorSize, SimdSizeX, SimdSizeY> A) \
 	{ \
 		TLaneVector2D<ReturnScalarType, VectorSize, SimdSizeX, SimdSizeY> R; \
-		UNROLL \
-		for (uint SimdIndex = 0; SimdIndex < SimdSizeX * SimdSizeY; SimdIndex++) \
-		{ \
-			R.SetElement(SimdIndex, FunctionName(A.GetElement(SimdIndex))); \
-		} \
-		return R; \
-	} \
-	\
-	CALL_SITE_DEBUGLOC \
-	template<typename ScalarType, uint VectorSize, uint SimdSizeX, uint SimdSizeY> \
-	TLanePackedVector2D<ReturnScalarType, VectorSize, SimdSizeX, SimdSizeY> FunctionName( \
-		TLanePackedVector2D<ScalarType, VectorSize, SimdSizeX, SimdSizeY> A) \
-	{ \
-		TLanePackedVector2D<ReturnScalarType, VectorSize, SimdSizeX, SimdSizeY> R; \
-		UNROLL \
-		for (uint ScalarPairIndex = 0; ScalarPairIndex < (VectorSize * SimdSizeX * SimdSizeY) / 2; ScalarPairIndex++) \
-		{ \
-			R.SetScalarPair(ScalarPairIndex, FunctionName(A.GetScalarPair(ScalarPairIndex))); \
+		UNROLL for (uint RegisterRowIndex = 0; RegisterRowIndex < TLaneVector2D<ScalarType, VectorSize, SimdSizeX, SimdSizeY>::RegisterRowCount; RegisterRowIndex++) { \
+			R.Registers.SetRegisterRow(RegisterRowIndex, FunctionName(A.Registers.GetRegisterRow(RegisterRowIndex))); \
 		} \
 		return R; \
 	} \
@@ -868,19 +1105,6 @@ TVECTOR_FUNCTION_1PARAMS(uint, asuint);
 // ------------- 2 parameters
 
 #define TVECTOR_FUNCTION_2PARAMS(FunctionName) \
-	CALL_SITE_DEBUGLOC \
-	template<typename ScalarType, uint VectorSize, uint SimdSize> \
-	TLaneVector<ScalarType, VectorSize, SimdSize> FunctionName( \
-		TLaneVector<ScalarType, VectorSize, SimdSize> A, \
-		TLaneVector<ScalarType, VectorSize, SimdSize> B) \
-	{ \
-		TLaneVector<ScalarType, VectorSize, SimdSize> R; \
-		UNROLL for (uint SimdIndex = 0; SimdIndex < SimdSize; SimdIndex++) { \
-			R.SetElement(SimdIndex, FunctionName(A.GetElement(SimdIndex), B.GetElement(SimdIndex))); \
-		} \
-		return R; \
-	} \
-	\
 	CALL_SITE_DEBUGLOC \
 	template<typename ScalarType, uint VectorSize, uint SimdSizeX, uint SimdSizeY> \
 	TLaneVector2D<ScalarType, VectorSize, SimdSizeX, SimdSizeY> FunctionName( \
@@ -888,22 +1112,8 @@ TVECTOR_FUNCTION_1PARAMS(uint, asuint);
 		TLaneVector2D<ScalarType, VectorSize, SimdSizeX, SimdSizeY> B) \
 	{ \
 		TLaneVector2D<ScalarType, VectorSize, SimdSizeX, SimdSizeY> R; \
-		UNROLL for (uint SimdIndex = 0; SimdIndex < SimdSizeX * SimdSizeY; SimdIndex++) { \
-			R.SetElement(SimdIndex, FunctionName(A.GetElement(SimdIndex), B.GetElement(SimdIndex))); \
-		} \
-		return R; \
-	} \
-	\
-	CALL_SITE_DEBUGLOC \
-	template<typename ScalarType, uint VectorSize, uint SimdSizeX, uint SimdSizeY> \
-	TLanePackedVector2D<ScalarType, VectorSize, SimdSizeX, SimdSizeY> FunctionName( \
-		TLanePackedVector2D<ScalarType, VectorSize, SimdSizeX, SimdSizeY> A, \
-		TLanePackedVector2D<ScalarType, VectorSize, SimdSizeX, SimdSizeY> B) \
-	{ \
-		TLanePackedVector2D<ScalarType, VectorSize, SimdSizeX, SimdSizeY> R; \
-		UNROLL for (uint ScalarPairIndex = 0; ScalarPairIndex < (VectorSize * SimdSizeX * SimdSizeY) / 2u; ScalarPairIndex++) \
-		{ \
-			R.SetScalarPair(ScalarPairIndex, FunctionName(A.GetScalarPair(ScalarPairIndex), B.GetScalarPair(ScalarPairIndex))); \
+		UNROLL for (uint RegisterRowIndex = 0; RegisterRowIndex < TLaneVector2D<ScalarType, VectorSize, SimdSizeX, SimdSizeY>::RegisterRowCount; RegisterRowIndex++) { \
+			R.Registers.SetRegisterRow(RegisterRowIndex, FunctionName(A.Registers.GetRegisterRow(RegisterRowIndex), B.Registers.GetRegisterRow(RegisterRowIndex))); \
 		} \
 		return R; \
 	} \
@@ -921,19 +1131,6 @@ TVECTOR_FUNCTION_2PARAMS(bit_shift_right);
 // ------------- 2 parameters but different return dimension
 
 #define TVECTOR_FUNCTION_2PARAMS_DIFF_RETURN(ReturnVectorSize, FunctionName) \
-	CALL_SITE_DEBUGLOC \
-	template<typename ScalarType, uint VectorSize, uint SimdSize> \
-	TLaneVector<ScalarType, ReturnVectorSize, SimdSize> FunctionName( \
-		TLaneVector<ScalarType, VectorSize, SimdSize> A, \
-		TLaneVector<ScalarType, VectorSize, SimdSize> B) \
-	{ \
-		TLaneVector<ScalarType, ReturnVectorSize, SimdSize> R; \
-		UNROLL for (uint SimdIndex = 0; SimdIndex < SimdSize; SimdIndex++) { \
-			R.SetElement(SimdIndex, FunctionName(A.GetElement(SimdIndex), B.GetElement(SimdIndex))); \
-		} \
-		return R; \
-	} \
-	\
 	CALL_SITE_DEBUGLOC \
 	template<typename ScalarType, uint VectorSize, uint SimdSizeX, uint SimdSizeY> \
 	TLaneVector2D<ScalarType, ReturnVectorSize, SimdSizeX, SimdSizeY> FunctionName( \
@@ -941,21 +1138,8 @@ TVECTOR_FUNCTION_2PARAMS(bit_shift_right);
 		TLaneVector2D<ScalarType, VectorSize, SimdSizeX, SimdSizeY> B) \
 	{ \
 		TLaneVector2D<ScalarType, ReturnVectorSize, SimdSizeX, SimdSizeY> R; \
-		UNROLL for (uint SimdIndex = 0; SimdIndex < SimdSizeX * SimdSizeY; SimdIndex++) { \
-			R.SetElement(SimdIndex, FunctionName(A.GetElement(SimdIndex), B.GetElement(SimdIndex))); \
-		} \
-		return R; \
-	} \
-	\
-	CALL_SITE_DEBUGLOC \
-	template<typename ScalarType, uint VectorSize, uint SimdSizeX, uint SimdSizeY> \
-	TLanePackedVector2D<ScalarType, ReturnVectorSize, SimdSizeX, SimdSizeY> FunctionName( \
-		TLanePackedVector2D<ScalarType, VectorSize, SimdSizeX, SimdSizeY> A, \
-		TLanePackedVector2D<ScalarType, VectorSize, SimdSizeX, SimdSizeY> B) \
-	{ \
-		TLanePackedVector2D<ScalarType, ReturnVectorSize, SimdSizeX, SimdSizeY> R; \
-		UNROLL for (uint SimdIndex = 0; SimdIndex < SimdSizeX * SimdSizeY; SimdIndex++) { \
-			R.SetElement(SimdIndex, FunctionName(A.GetElement(SimdIndex), B.GetElement(SimdIndex))); \
+		UNROLL for (uint ElementIndex = 0; ElementIndex < SimdSizeX * SimdSizeY; ElementIndex++) { \
+			R.SetElement(ElementIndex, FunctionName(A.GetElement(ElementIndex), B.GetElement(ElementIndex))); \
 		} \
 		return R; \
 	} \
@@ -965,20 +1149,6 @@ TVECTOR_FUNCTION_2PARAMS_DIFF_RETURN(/* VectorSize = */ 1, dot);
 
 // ------------- 3 parameters
 #define TVECTOR_FUNCTION_3PARAMS(FunctionName, ScalarTypeA) \
-	CALL_SITE_DEBUGLOC \
-	template<typename ScalarType, uint VectorSize, uint SimdSize> \
-	TLaneVector<ScalarType, VectorSize, SimdSize> FunctionName( \
-		TLaneVector<ScalarTypeA, VectorSize, SimdSize> A, \
-		TLaneVector<ScalarType, VectorSize, SimdSize> B, \
-		TLaneVector<ScalarType, VectorSize, SimdSize> C) \
-	{ \
-		TLaneVector<ScalarType, VectorSize, SimdSize> R; \
-		UNROLL for (uint SimdIndex = 0; SimdIndex < SimdSize; SimdIndex++) { \
-			R.SetElement(SimdIndex, FunctionName(A.GetElement(SimdIndex), B.GetElement(SimdIndex), C.GetElement(SimdIndex))); \
-		} \
-		return R; \
-	} \
-	\
 	CALL_SITE_DEBUGLOC \
 	template<typename ScalarType, uint VectorSize, uint SimdSizeX, uint SimdSizeY> \
 	TLaneVector2D<ScalarType, VectorSize, SimdSizeX, SimdSizeY> FunctionName( \
@@ -987,23 +1157,8 @@ TVECTOR_FUNCTION_2PARAMS_DIFF_RETURN(/* VectorSize = */ 1, dot);
 		TLaneVector2D<ScalarType, VectorSize, SimdSizeX, SimdSizeY> C) \
 	{ \
 		TLaneVector2D<ScalarType, VectorSize, SimdSizeX, SimdSizeY> R; \
-		UNROLL for (uint SimdIndex = 0; SimdIndex < SimdSizeX * SimdSizeY; SimdIndex++) { \
-			R.SetElement(SimdIndex, FunctionName(A.GetElement(SimdIndex), B.GetElement(SimdIndex), C.GetElement(SimdIndex))); \
-		} \
-		return R; \
-	} \
-	\
-	CALL_SITE_DEBUGLOC \
-	template<typename ScalarType, uint VectorSize, uint SimdSizeX, uint SimdSizeY> \
-	TLanePackedVector2D<ScalarType, VectorSize, SimdSizeX, SimdSizeY> FunctionName( \
-		TLanePackedVector2D<ScalarTypeA, VectorSize, SimdSizeX, SimdSizeY> A, \
-		TLanePackedVector2D<ScalarType, VectorSize, SimdSizeX, SimdSizeY> B, \
-		TLanePackedVector2D<ScalarType, VectorSize, SimdSizeX, SimdSizeY> C) \
-	{ \
-		TLanePackedVector2D<ScalarType, VectorSize, SimdSizeX, SimdSizeY> R; \
-		UNROLL for (uint ScalarPairIndex = 0; ScalarPairIndex < (VectorSize * SimdSizeX * SimdSizeY) / 2u; ScalarPairIndex++) \
-		{ \
-			R.SetScalarPair(ScalarPairIndex, FunctionName(A.GetScalarPair(ScalarPairIndex), B.GetScalarPair(ScalarPairIndex), C.GetScalarPair(ScalarPairIndex))); \
+		UNROLL for (uint RegisterRowIndex = 0; RegisterRowIndex < TLaneVector2D<ScalarType, VectorSize, SimdSizeX, SimdSizeY>::RegisterRowCount; RegisterRowIndex++) { \
+			R.Registers.SetRegisterRow(RegisterRowIndex, FunctionName(A.Registers.GetRegisterRow(RegisterRowIndex), B.Registers.GetRegisterRow(RegisterRowIndex), C.Registers.GetRegisterRow(RegisterRowIndex))); \
 		} \
 		return R; \
 	} \
diff --git a/Engine/Shaders/Public/WaveBroadcastIntrinsics.ush b/Engine/Shaders/Public/WaveBroadcastIntrinsics.ush
--- a/Engine/Shaders/Public/WaveBroadcastIntrinsics.ush
+++ b/Engine/Shaders/Public/WaveBroadcastIntrinsics.ush
@@ -264,6 +264,40 @@ float4 WaveBroadcast(const FWaveBroadcastSettings Settings, float4 v)
 } 
 
 
+CALL_SITE_DEBUGLOC
+bool WaveBroadcast(const FWaveBroadcastSettings Settings, bool x)
+{
+	return bool(WaveBroadcast(Settings, uint(x)));
+}
+
+CALL_SITE_DEBUGLOC
+bool2 WaveBroadcast(const FWaveBroadcastSettings Settings, bool2 v)
+{
+	return bool2(
+		WaveBroadcast(Settings, v.x),
+		WaveBroadcast(Settings, v.y));
+} 
+
+CALL_SITE_DEBUGLOC
+bool3 WaveBroadcast(const FWaveBroadcastSettings Settings, bool3 v)
+{
+	return bool3(
+		WaveBroadcast(Settings, v.x),
+		WaveBroadcast(Settings, v.y),
+		WaveBroadcast(Settings, v.z));
+}
+
+CALL_SITE_DEBUGLOC
+bool4 WaveBroadcast(const FWaveBroadcastSettings Settings, bool4 v)
+{
+	return bool4(
+		WaveBroadcast(Settings, v.x),
+		WaveBroadcast(Settings, v.y),
+		WaveBroadcast(Settings, v.z),
+		WaveBroadcast(Settings, v.w));
+} 
+
+
 #if PLATFORM_SUPPORTS_REAL_TYPES
 
 CALL_SITE_DEBUGLOC
