Concatenates the convolutions of the TSR rejection's blend and clamp to save perf

diff --git a/Engine/Shaders/Private/TemporalSuperResolution/TSRConvolutionNetwork.ush b/Engine/Shaders/Private/TemporalSuperResolution/TSRConvolutionNetwork.ush
--- a/Engine/Shaders/Private/TemporalSuperResolution/TSRConvolutionNetwork.ush
+++ b/Engine/Shaders/Private/TemporalSuperResolution/TSRConvolutionNetwork.ush
@@ -712,6 +712,28 @@ TLaneVector2D<ScalarType, VectorSize, LaneStrideX, LaneStrideY> Max3x3(TLaneVect
 	return max3(MaxN, MaxP, MaxO);
 }
 
+CALL_SITE_DEBUGLOC
+template<typename ScalarType, uint LaneStrideX, uint LaneStrideY>
+TLaneVector2D<ScalarType, 4, LaneStrideX, LaneStrideY> MaxRGBMinA3x3(TLaneVector2D<ScalarType, 4, LaneStrideX, LaneStrideY> Center)
+{
+	TLaneVector2D<ScalarType, 4, LaneStrideX, LaneStrideY> CO = Center;
+	TLaneVector2D<ScalarType, 4, LaneStrideX, LaneStrideY> CP;
+	TLaneVector2D<ScalarType, 4, LaneStrideX, LaneStrideY> CN;
+	AccessNeighborTexels1x3(CO, /* out */ CP, /* out */ CN);
+	
+	TLaneVector2D<ScalarType, 4, LaneStrideX, LaneStrideY> MaxO = max3(CN, CP, CO);
+	MaxO.SetComponent(3, min3(CN[3], CP[3], CO[3]));
+
+	TLaneVector2D<ScalarType, 4, LaneStrideX, LaneStrideY> MaxP;
+	TLaneVector2D<ScalarType, 4, LaneStrideX, LaneStrideY> MaxN;
+	AccessNeighborTexels3x1(MaxO, /* out */ MaxP, /* out */ MaxN);
+	
+	TLaneVector2D<ScalarType, 4, LaneStrideX, LaneStrideY> R = max3(MaxN, MaxP, MaxO);
+	R.SetComponent(3, min3(MaxN[3], MaxP[3], MaxO[3]));
+
+	return R;
+}
+
 CALL_SITE_DEBUGLOC
 template<typename ScalarType, uint VectorSize, uint LaneStrideX, uint LaneStrideY>
 void MinMax3x3(
diff --git a/Engine/Shaders/Private/TemporalSuperResolution/TSRRejectShading.usf b/Engine/Shaders/Private/TemporalSuperResolution/TSRRejectShading.usf
--- a/Engine/Shaders/Private/TemporalSuperResolution/TSRRejectShading.usf
+++ b/Engine/Shaders/Private/TemporalSuperResolution/TSRRejectShading.usf
@@ -405,22 +405,26 @@ void MeasureRejection(
 		}
 	}
 
-	tsr_simd_halfC RawClampedEnergy      = abs(ClampedFilteredHistory - FilteredHistory);
-	tsr_simd_halfC FilteredClampedEnergy = Max3x3(Median3x3(RawClampedEnergy));
-
 	tsr_simd_halfC Delta = abs(FilteredInput - FilteredHistory);
 	tsr_simd_halfC ClampError = max(
 		BoxSize * tsr_half(0.25),
 		BackbufferQuantizationErrorVector);
 			
+	tsr_simd_halfC RawClampedEnergy      = abs(ClampedFilteredHistory - FilteredHistory);
 	tsr_simd_halfC RawFactor      = saturate(tsr_simd_halfC::Const(1.0) - max(RawClampedEnergy - ClampError, tsr_simd_halfC::Const(0.0)) * rcp(max(Delta, BoxSize)));
-	tsr_simd_halfC FilteredFactor = saturate(tsr_simd_halfC::Const(1.0) - max(FilteredClampedEnergy - ClampError, tsr_simd_halfC::Const(0.0)) * rcp(max(Delta, BoxSize)));
-		
 	tsr_simd_half RawRejection      = min3(RawFactor[0], RawFactor[1], RawFactor[2]);
-	tsr_simd_half FilteredRejection = min3(FilteredFactor[0], FilteredFactor[1], FilteredFactor[2]);
 
+	#if CONFIG_SCENE_COLOR_ALPHA
+		OutRejectionClampBlend = Min3x3(Median3x3(RawRejection));
+		tsr_simd_halfC FilteredClampedEnergy = Max3x3(Median3x3(RawClampedEnergy));
+	#else
+		tsr_simd_halfC FilteredClampedEnergy;
+		Deconcatenate(MaxRGBMinA3x3(Median3x3(Concatenate(RawClampedEnergy, RawRejection))), /* out */ FilteredClampedEnergy, /* out */ OutRejectionClampBlend);
+	#endif
+
+	tsr_simd_halfC FilteredFactor = saturate(tsr_simd_halfC::Const(1.0) - max(FilteredClampedEnergy - ClampError, tsr_simd_halfC::Const(0.0)) * rcp(max(Delta, BoxSize)));
+	tsr_simd_half FilteredRejection = min3(FilteredFactor[0], FilteredFactor[1], FilteredFactor[2]);
 	OutRejectionBlendFinal = saturate(FilteredRejection);
-	OutRejectionClampBlend = Min3x3(Median3x3(RawRejection));
 } // MeasureRejection()
 
 
diff --git a/Engine/Shaders/Public/LaneVectorization.ush b/Engine/Shaders/Public/LaneVectorization.ush
--- a/Engine/Shaders/Public/LaneVectorization.ush
+++ b/Engine/Shaders/Public/LaneVectorization.ush
@@ -1073,6 +1073,66 @@ struct TLaneVector2D
 }; // TLaneVector2D
 
 
+CALL_SITE_DEBUGLOC
+template<typename ScalarType, uint VectorSizeA, uint VectorSizeB, uint SimdSizeX, uint SimdSizeY>
+void Concatenate(
+	TLaneVector2D<ScalarType, VectorSizeA, SimdSizeX, SimdSizeY> A,
+	TLaneVector2D<ScalarType, VectorSizeB, SimdSizeX, SimdSizeY> B,
+	out TLaneVector2D<ScalarType, VectorSizeA + VectorSizeB, SimdSizeX, SimdSizeY> R)
+{
+	UNROLL
+	for (uint ElementIndex = 0; ElementIndex < SimdSizeX * SimdSizeY; ElementIndex++)
+	{
+		UNROLL
+		for (uint ComponentIndex = 0; ComponentIndex < VectorSizeA; ComponentIndex++)
+		{
+			R.Registers.SetElementComponent(ElementIndex, ComponentIndex, A.Registers.GetElementComponent(ElementIndex, ComponentIndex));
+		}
+
+		UNROLL
+		for (uint ComponentIndex = 0; ComponentIndex < VectorSizeB; ComponentIndex++)
+		{
+			R.Registers.SetElementComponent(ElementIndex, VectorSizeA + ComponentIndex, B.Registers.GetElementComponent(ElementIndex, ComponentIndex));
+		}
+	}
+}
+
+CALL_SITE_DEBUGLOC
+template<typename ScalarType, uint VectorSizeA, uint VectorSizeB, uint SimdSizeX, uint SimdSizeY>
+TLaneVector2D<ScalarType, VectorSizeA + VectorSizeB, SimdSizeX, SimdSizeY> Concatenate(
+	TLaneVector2D<ScalarType, VectorSizeA, SimdSizeX, SimdSizeY> A,
+	TLaneVector2D<ScalarType, VectorSizeB, SimdSizeX, SimdSizeY> B)
+{
+	TLaneVector2D<ScalarType, VectorSizeA + VectorSizeB, SimdSizeX, SimdSizeY> R;
+	Concatenate(A, B, /* out */ R);
+	return R;
+}
+
+CALL_SITE_DEBUGLOC
+template<typename ScalarType, uint VectorSizeA, uint VectorSizeB, uint SimdSizeX, uint SimdSizeY>
+void Deconcatenate(
+	TLaneVector2D<ScalarType, VectorSizeA + VectorSizeB, SimdSizeX, SimdSizeY> M,
+	out TLaneVector2D<ScalarType, VectorSizeA, SimdSizeX, SimdSizeY> A,
+	out TLaneVector2D<ScalarType, VectorSizeB, SimdSizeX, SimdSizeY> B)
+{
+	UNROLL
+	for (uint ElementIndex = 0; ElementIndex < SimdSizeX * SimdSizeY; ElementIndex++)
+	{
+		UNROLL
+		for (uint ComponentIndex = 0; ComponentIndex < VectorSizeA; ComponentIndex++)
+		{
+			A.Registers.SetElementComponent(ElementIndex, ComponentIndex, M.Registers.GetElementComponent(ElementIndex, ComponentIndex));
+		}
+
+		UNROLL
+		for (uint ComponentIndex = 0; ComponentIndex < VectorSizeB; ComponentIndex++)
+		{
+			B.Registers.SetElementComponent(ElementIndex, ComponentIndex, M.Registers.GetElementComponent(ElementIndex, VectorSizeA + ComponentIndex));
+		}
+	}
+}
+
+
 // ------------- 1 parameter
 
 #define TVECTOR_FUNCTION_1PARAMS(ReturnScalarType, FunctionName) \
