PR #10788: Fixed undefined RasterizerOrderedTexture2D in Vulkan

#rnx

diff --git a/Engine/Shaders/Public/Platform.ush b/Engine/Shaders/Public/Platform.ush
--- a/Engine/Shaders/Public/Platform.ush
+++ b/Engine/Shaders/Public/Platform.ush
@@ -143,6 +143,9 @@
 
 #ifndef PLATFORM_SUPPORTS_ROV
 #define PLATFORM_SUPPORTS_ROV 0
+#endif
+
+#if !PLATFORM_SUPPORTS_ROV
 #define RasterizerOrderedTexture2D RWTexture2D
 #endif
 
