Moving UB_MEMBER_UNIFORM_BLOCK to Platform.ush

diff --git a/Engine/Shaders/Public/BindlessResources.ush b/Engine/Shaders/Public/BindlessResources.ush
--- a/Engine/Shaders/Public/BindlessResources.ush
+++ b/Engine/Shaders/Public/BindlessResources.ush
@@ -48,10 +48,6 @@
 	#define UB_RESOURCE_MEMBER_RESOURCE(Type, Name) Type Name
 #endif
 
-#ifndef UB_MEMBER_UNIFORM_BLOCK
-	#define UB_MEMBER_UNIFORM_BLOCK(Name) float4 Name[PLATFORM_MAX_UNIFORM_BUFFER_RANGE_FLOAT4]
-#endif
-
 #if ENABLE_BINDLESS_SAMPLERS && PLATFORM_SUPPORTS_BINDLESS
 	#define DEFINE_SAMPLER_VARIABLE(Type, Name) \
 		Type GetBindlessSampler##Name() \
diff --git a/Engine/Shaders/Public/Platform.ush b/Engine/Shaders/Public/Platform.ush
--- a/Engine/Shaders/Public/Platform.ush
+++ b/Engine/Shaders/Public/Platform.ush
@@ -1297,7 +1297,9 @@ uint WaveGetActiveLaneIndexLast()
 
 // 16KB by default. Must match Platform.h
 #ifndef PLATFORM_MAX_UNIFORM_BUFFER_RANGE
-#define PLATFORM_MAX_UNIFORM_BUFFER_RANGE (16u*1024)
+	#define PLATFORM_MAX_UNIFORM_BUFFER_RANGE (16u*1024)
 #endif
 #define PLATFORM_MAX_UNIFORM_BUFFER_RANGE_FLOAT4 (PLATFORM_MAX_UNIFORM_BUFFER_RANGE/16u)
 
+#define UB_MEMBER_UNIFORM_BLOCK(Name) float4 Name[PLATFORM_MAX_UNIFORM_BUFFER_RANGE_FLOAT4]
+
