Fixes TSR history resurrection on D3D11

diff --git a/Engine/Shaders/Public/LaneVectorization.ush b/Engine/Shaders/Public/LaneVectorization.ush
--- a/Engine/Shaders/Public/LaneVectorization.ush
+++ b/Engine/Shaders/Public/LaneVectorization.ush
@@ -341,9 +341,9 @@ uint ReadDwordFromLDS(const uint SharedIndex)
 	return SharedData[SharedIndex];
 }
 
-void AtomicIncrementLDSDword(const uint SharedIndex, uint Value)
+void AtomicIncrementLDSDword(const uint SharedIndex, uint V)
 {
-	return InterlockedAdd(/* inout */ SharedData[SharedIndex], Value);
+	InterlockedAdd(/* inout */ SharedData[SharedIndex], V);
 }
 
 // uint
