Implements manual LDS despilling in TSR's rejectShading pass to save register pressure

diff --git a/Engine/Shaders/Private/TemporalSuperResolution/TSRConvolutionNetwork.ush b/Engine/Shaders/Private/TemporalSuperResolution/TSRConvolutionNetwork.ush
--- a/Engine/Shaders/Private/TemporalSuperResolution/TSRConvolutionNetwork.ush
+++ b/Engine/Shaders/Private/TemporalSuperResolution/TSRConvolutionNetwork.ush
@@ -43,11 +43,6 @@
  *    o   o   o   o | o   o   o   o
  */
 
- 
-//------------------------------------------------------- GLOBAL VARIABLE
-
-static uint GGroupThreadIndex = 0;
-
 
 //------------------------------------------------------- VECTOR MEMORY ORDER
 
@@ -321,7 +316,7 @@ void AccessNeighborTexels3x1(
 		UNROLL
 		for (uint SimdIndex = 0; SimdIndex < LaneStrideY; SimdIndex++)
 		{
-			WriteVectorToLDS(GGroupThreadIndex + SimdIndex * TotalNumberOfLaneInGroup, CO.GetElement(SimdIndex));
+			WriteVectorToLDS<ScalarType, VectorSize>(GGroupThreadIndex + SimdIndex * TotalNumberOfLaneInGroup, CO.GetElement(SimdIndex));
 		}
 		GroupMemoryBarrierWithGroupSync();
 		
@@ -330,8 +325,8 @@ void AccessNeighborTexels3x1(
 		{
 			vector<ScalarType, VectorSize> FirstElement;
 			vector<ScalarType, VectorSize> LastElement;
-			ReadVectorFromLDS(((GGroupThreadIndex - 1) % TotalNumberOfLaneInGroup) + SimdIndex * TotalNumberOfLaneInGroup, /* out */ FirstElement);
-			ReadVectorFromLDS(((GGroupThreadIndex + 1) % TotalNumberOfLaneInGroup) + SimdIndex * TotalNumberOfLaneInGroup, /* out */ LastElement);
+			ReadVectorFromLDS<ScalarType, VectorSize>(((GGroupThreadIndex - 1) % TotalNumberOfLaneInGroup) + SimdIndex * TotalNumberOfLaneInGroup, /* out */ FirstElement);
+			ReadVectorFromLDS<ScalarType, VectorSize>(((GGroupThreadIndex + 1) % TotalNumberOfLaneInGroup) + SimdIndex * TotalNumberOfLaneInGroup, /* out */ LastElement);
 	
 			CP.SetElement(SimdIndex, LastElement);
 			CN.SetElement(SimdIndex, FirstElement);
@@ -371,7 +366,7 @@ void AccessNeighborTexels1x3(
 		UNROLL
 		for (uint SimdIndex = 0; SimdIndex < LaneStrideX; SimdIndex++)
 		{
-			WriteVectorToLDS(GGroupThreadIndex + SimdIndex * TotalNumberOfLaneInGroup, CO.GetElement(SimdIndex));
+			WriteVectorToLDS<ScalarType, VectorSize>(GGroupThreadIndex + SimdIndex * TotalNumberOfLaneInGroup, CO.GetElement(SimdIndex));
 		}
 		GroupMemoryBarrierWithGroupSync();
 		
@@ -380,8 +375,8 @@ void AccessNeighborTexels1x3(
 		{
 			vector<ScalarType, VectorSize> FirstElement;
 			vector<ScalarType, VectorSize> LastElement;
-			ReadVectorFromLDS(((GGroupThreadIndex - LANE_COUNT_X * WAVE_COUNT_X) % TotalNumberOfLaneInGroup) + SimdIndex * TotalNumberOfLaneInGroup, /* out */ FirstElement);
-			ReadVectorFromLDS(((GGroupThreadIndex + LANE_COUNT_X * WAVE_COUNT_X) % TotalNumberOfLaneInGroup) + SimdIndex * TotalNumberOfLaneInGroup, /* out */ LastElement);
+			ReadVectorFromLDS<ScalarType, VectorSize>(((GGroupThreadIndex - LANE_COUNT_X * WAVE_COUNT_X) % TotalNumberOfLaneInGroup) + SimdIndex * TotalNumberOfLaneInGroup, /* out */ FirstElement);
+			ReadVectorFromLDS<ScalarType, VectorSize>(((GGroupThreadIndex + LANE_COUNT_X * WAVE_COUNT_X) % TotalNumberOfLaneInGroup) + SimdIndex * TotalNumberOfLaneInGroup, /* out */ LastElement);
 	
 			CP.SetElement(SimdIndex, LastElement);
 			CN.SetElement(SimdIndex, FirstElement);
@@ -396,8 +391,8 @@ void AccessNeighborTexels1x3(
 		UNROLL
 		for (uint SimdIndex = 0; SimdIndex < LaneStrideX; SimdIndex++)
 		{
-			WriteVectorToLDS(GGroupThreadIndex + (0 + SimdIndex) * TotalNumberOfLaneInGroup, CO.GetElement(FirstElementIndex + SimdIndex));
-			WriteVectorToLDS(GGroupThreadIndex + (LaneStrideX + SimdIndex) * TotalNumberOfLaneInGroup, CO.GetElement(LastElementIndex + SimdIndex));
+			WriteVectorToLDS<ScalarType, VectorSize>(GGroupThreadIndex + (0 + SimdIndex) * TotalNumberOfLaneInGroup, CO.GetElement(FirstElementIndex + SimdIndex));
+			WriteVectorToLDS<ScalarType, VectorSize>(GGroupThreadIndex + (LaneStrideX + SimdIndex) * TotalNumberOfLaneInGroup, CO.GetElement(LastElementIndex + SimdIndex));
 		}
 		GroupMemoryBarrierWithGroupSync();
 	
@@ -413,8 +408,8 @@ void AccessNeighborTexels1x3(
 		{
 			vector<ScalarType, VectorSize> FirstElement;
 			vector<ScalarType, VectorSize> LastElement;
-			ReadVectorFromLDS(((GGroupThreadIndex - LANE_COUNT_X * WAVE_COUNT_X) % TotalNumberOfLaneInGroup) + (LaneStrideX + SimdIndex) * TotalNumberOfLaneInGroup, /* out */ FirstElement);
-			ReadVectorFromLDS(((GGroupThreadIndex + LANE_COUNT_X * WAVE_COUNT_X) % TotalNumberOfLaneInGroup) + (0 + SimdIndex) * TotalNumberOfLaneInGroup, /* out */ LastElement);
+			ReadVectorFromLDS<ScalarType, VectorSize>(((GGroupThreadIndex - LANE_COUNT_X * WAVE_COUNT_X) % TotalNumberOfLaneInGroup) + (LaneStrideX + SimdIndex) * TotalNumberOfLaneInGroup, /* out */ FirstElement);
+			ReadVectorFromLDS<ScalarType, VectorSize>(((GGroupThreadIndex + LANE_COUNT_X * WAVE_COUNT_X) % TotalNumberOfLaneInGroup) + (0 + SimdIndex) * TotalNumberOfLaneInGroup, /* out */ LastElement);
 	
 			CP.SetElement(LastElementIndex + SimdIndex, LastElement);
 			CN.SetElement(FirstElementIndex + SimdIndex, FirstElement);
diff --git a/Engine/Shaders/Private/TemporalSuperResolution/TSRRejectShading.usf b/Engine/Shaders/Private/TemporalSuperResolution/TSRRejectShading.usf
--- a/Engine/Shaders/Private/TemporalSuperResolution/TSRRejectShading.usf
+++ b/Engine/Shaders/Private/TemporalSuperResolution/TSRRejectShading.usf
@@ -58,6 +58,11 @@
 	#ifdef PLATFORM_SPECIFIC_TENSOR_REGISTER_LAYOUT
 		#define TENSOR_REGISTER_LAYOUT PLATFORM_SPECIFIC_TENSOR_REGISTER_LAYOUT
 	#endif
+	
+	/** Manually despill to LDS on RDNA platforms with history resurrection to save more register pressure. */
+	#if CONFIG_MINIMISE_VGPRS && CONFIG_COMPILE_FP16 && !CONFIG_SCENE_COLOR_ALPHA
+		#define CONFIG_LDS_DESPILL 1
+	#endif
 
 #elif DIM_WAVE_SIZE == 0
 	/** Load an entire 32x32 with 1x1 worth of data by lane, only using LDS. */
@@ -86,6 +91,24 @@
 	#define LDS_DWORD_COMPONENT_COUNT 4
 #endif
 
+#ifndef CONFIG_LDS_DESPILL
+	#define CONFIG_LDS_DESPILL 0
+#endif
+
+/** Manually despill to LDS on RDNA platforms to save more register pressure. */
+#if CONFIG_LDS_DESPILL
+	#define LDS_DESPILL_THREAD_COUNT (WAVE_COUNT_X * WAVE_COUNT_Y) * (LANE_COUNT_X * LANE_COUNT_Y)
+	#define LDS_DESPILL_OFFSET       (LANE_STRIDE_X * LANE_STRIDE_Y) * GetVectorDwordSize(tsr_half, 4)
+
+	#if CONFIG_HISTORY_RESURRECTION
+		#define LDS_DESPILL_OFFSET_COUNT 2
+	#else
+		#define LDS_DESPILL_OFFSET_COUNT 1
+	#endif
+	
+	#define LDS_DESPILL_DWORD_COUNT  LDS_DESPILL_THREAD_COUNT * LDS_DESPILL_OFFSET * LDS_DESPILL_OFFSET_COUNT
+#endif
+
 /** Size of the SIMD per lane that also number of input pixel loaded into each individual lanes. */
 #define SIMD_SIZE (LANE_STRIDE_X * LANE_STRIDE_Y)
 
@@ -747,6 +770,29 @@ void MainCS(
 			/* SourcePerceptionAdd = */ 1.0,
 			/* DestPerceptionAdd   = */ PerceptionAdd);
 			
+		// Despill to LDS to make room in VGPRs for the convolutions of MeasureRejection()
+		#if CONFIG_LDS_DESPILL
+		const uint kDespillHistory = LDS_DESPILL_THREAD_COUNT * LDS_DESPILL_OFFSET * 0;
+		#if CONFIG_HISTORY_RESURRECTION
+			const uint kDespillResurrectedHistory = LDS_DESPILL_THREAD_COUNT * LDS_DESPILL_OFFSET * 1;
+		#endif
+		PLATFORM_SPECIFIC_ISOLATE
+		{
+			Concatenate(ExposedHistory, OriginalHistoryUncertainty).StoreLDSSpill(kDespillHistory);
+			
+			#if CONFIG_HISTORY_RESURRECTION
+			{
+				tsr_simd_halfC ExposedResurrectedHistory = ChangePerceptionAdd(
+					OriginalResurrectedHistory,
+					/* SourcePerceptionAdd = */ 1.0,
+					/* DestPerceptionAdd   = */ PerceptionAdd);
+				
+				Concatenate(ExposedResurrectedHistory, OriginalResurrectedHistoryUncertainty).StoreLDSSpill(kDespillResurrectedHistory);
+			}
+			#endif
+		}
+		#endif // CONFIG_LDS_DESPILL
+
 		tsr_simd_halfC ExposedInputMin;
 		tsr_simd_halfC ExposedInputMax;
 		MinMax3x3(ExposedInput, /* out */ ExposedInputMin, /* out */ ExposedInputMax);
@@ -780,6 +826,14 @@ void MainCS(
 				/* SourcePerceptionAdd = */ 1.0,
 				/* DestPerceptionAdd   = */ PerceptionAdd);
 				
+			#if CONFIG_LDS_DESPILL
+			PLATFORM_SPECIFIC_ISOLATE
+			{
+				Deconcatenate(tsr_simd_half4::LoadLDSSpill(kDespillHistory), /* out */ ExposedHistory, /* out */ OriginalHistoryUncertainty);
+				Deconcatenate(tsr_simd_half4::LoadLDSSpill(kDespillResurrectedHistory), /* out */ ExposedResurrectedHistory, /* out */ OriginalResurrectedHistoryUncertainty);
+			}
+			#endif // CONFIG_LDS_DESPILL
+
 			tsr_simd_halfC PrevMatch         = abs(ExposedInput - ExposedHistory);
 			tsr_simd_halfC ResurrectionMatch = abs(ExposedInput - ExposedResurrectedHistory);
 
@@ -812,6 +866,14 @@ void MainCS(
 				tsr_simd_bool bResurrectionIsBetter = (ResurrectionRejectionBlendFinal - RejectionBlendFinal) > tsr_simd_half::Const(0.1);
 				bResurrectHistory = and(bResurrectionIsCloser, bResurrectionIsBetter);
 				
+				#if CONFIG_LDS_DESPILL
+				PLATFORM_SPECIFIC_ISOLATE
+				{
+					Deconcatenate(tsr_simd_half4::LoadLDSSpill(kDespillHistory), /* out */ ExposedHistory, /* out */ OriginalHistoryUncertainty);
+					Deconcatenate(tsr_simd_half4::LoadLDSSpill(kDespillResurrectedHistory), /* out */ ExposedResurrectedHistory, /* out */ OriginalResurrectedHistoryUncertainty);
+				}
+				#endif // CONFIG_LDS_DESPILL
+
 				RejectionBlendFinal = select(bResurrectHistory, ResurrectionRejectionBlendFinal, RejectionBlendFinal);
 				RejectionClampBlend = select(bResurrectHistory, ResurrectionRejectionClampBlend, RejectionClampBlend);
 				OriginalHistoryUncertainty = select(bResurrectHistory, OriginalResurrectedHistoryUncertainty, OriginalHistoryUncertainty);
@@ -823,7 +885,12 @@ void MainCS(
 				#endif
 			}
 		}
-		#endif
+		#elif !CONFIG_HISTORY_RESURRECTION && CONFIG_LDS_DESPILL
+		PLATFORM_SPECIFIC_ISOLATE
+		{
+			Deconcatenate(tsr_simd_half4::LoadLDSSpill(kDespillHistory), /* out */ ExposedHistory, /* out */ OriginalHistoryUncertainty);
+		}
+		#endif //!CONFIG_HISTORY_RESURRECTION
 
 		// Recomputes original off the exposed shading to save register pressure.
 		#if CONFIG_MINIMISE_VGPRS
diff --git a/Engine/Shaders/Public/LaneVectorization.ush b/Engine/Shaders/Public/LaneVectorization.ush
--- a/Engine/Shaders/Public/LaneVectorization.ush
+++ b/Engine/Shaders/Public/LaneVectorization.ush
@@ -21,6 +21,11 @@
 #endif
 
 
+//------------------------------------------------------- GLOBAL VARIABLE
+
+static uint GGroupThreadIndex = 0;
+
+
 //------------------------------------------------------- STANDARD SCALAR OPERATIONS
 
 bool not(bool x)
@@ -325,133 +330,118 @@ int16_t4 PackRegistersTightly(int16_t4 v)
 #endif // PLATFORM_SUPPORTS_REAL_TYPES
 
 
-//------------------------------------------------------- GROUP SHARED COMMUNICATION
-
-#if defined(LDS_SIZE)
-
-groupshared uint SharedData[LDS_SIZE * LDS_DWORD_COMPONENT_COUNT];
-
-void WriteDwordToLDS(const uint SharedIndex, uint V)
-{
-	SharedData[SharedIndex] = V;
-}
+//------------------------------------------------------- DWORD PACKING/UNPACKING
 
-uint ReadDwordFromLDS(const uint SharedIndex)
-{
-	return SharedData[SharedIndex];
-}
+/** Returns the byte size of vector<ScalarType, VectorSize> */
+#define GetVectorByteSize(ScalarType, VectorSize) uint(sizeof(ScalarType) * uint(VectorSize))
 
-void AtomicIncrementLDSDword(const uint SharedIndex, uint V)
-{
-	InterlockedAdd(/* inout */ SharedData[SharedIndex], V);
-}
+/** Returns the dowrd size of vector<ScalarType, VectorSize> */
+#define GetVectorDwordSize(ScalarType, VectorSize) (uint(GetVectorByteSize(ScalarType, VectorSize) + 3u) / 4u)
 
 // uint
-void WriteVectorToLDS(const uint SharedIndex, uint V)
+void PackVectorToDwords(uint V, out uint DW[1])
 {
-	WriteDwordToLDS(SharedIndex, V);
+	DW[0] = V;
 }
 
-void WriteVectorToLDS(const uint SharedIndex, uint2 V)
+void PackVectorToDwords(uint2 V, out uint DW[2])
 {
-	WriteDwordToLDS(SharedIndex + 0 * LDS_SIZE, V.x);
-	WriteDwordToLDS(SharedIndex + 1 * LDS_SIZE, V.y);
+	DW[0] = V.x;
+	DW[1] = V.y;
 }
 
-void WriteVectorToLDS(const uint SharedIndex, uint3 V)
+void PackVectorToDwords(uint3 V, out uint DW[3])
 {
-	WriteDwordToLDS(SharedIndex + 0 * LDS_SIZE, V.x);
-	WriteDwordToLDS(SharedIndex + 1 * LDS_SIZE, V.y);
-	WriteDwordToLDS(SharedIndex + 2 * LDS_SIZE, V.z);
+	DW[0] = V.x;
+	DW[1] = V.y;
+	DW[2] = V.z;
 }
 
-void WriteVectorToLDS(const uint SharedIndex, uint4 V)
+void PackVectorToDwords(uint4 V, out uint DW[4])
 {
-	WriteDwordToLDS(SharedIndex + 0 * LDS_SIZE, V.x);
-	WriteDwordToLDS(SharedIndex + 1 * LDS_SIZE, V.y);
-	WriteDwordToLDS(SharedIndex + 2 * LDS_SIZE, V.z);
-	WriteDwordToLDS(SharedIndex + 3 * LDS_SIZE, V.w);
+	DW[0] = V.x;
+	DW[1] = V.y;
+	DW[2] = V.z;
+	DW[3] = V.w;
 }
 
-void ReadVectorFromLDS(const uint SharedIndex, out uint V)
+void UnpackDwordsToVector(uint DW[1], out uint V)
 {
-	V = ReadDwordFromLDS(SharedIndex);
+	V   = DW[0];
 }
 
-void ReadVectorFromLDS(const uint SharedIndex, out uint2 V)
+void UnpackDwordsToVector(uint DW[2], out uint2 V)
 {
-	V.x = ReadDwordFromLDS(SharedIndex + 0 * LDS_SIZE);
-	V.y = ReadDwordFromLDS(SharedIndex + 1 * LDS_SIZE);
+	V.x = DW[0];
+	V.y = DW[1];
 }
 
-void ReadVectorFromLDS(const uint SharedIndex, out uint3 V)
+void UnpackDwordsToVector(uint DW[3], out uint3 V)
 {
-	V.x = ReadDwordFromLDS(SharedIndex + 0 * LDS_SIZE);
-	V.y = ReadDwordFromLDS(SharedIndex + 1 * LDS_SIZE);
-	V.z = ReadDwordFromLDS(SharedIndex + 2 * LDS_SIZE);
+	V.x = DW[0];
+	V.y = DW[1];
+	V.z = DW[2];
 }
 
-void ReadVectorFromLDS(const uint SharedIndex, out uint4 V)
+void UnpackDwordsToVector(uint DW[4], out uint4 V)
 {
-	V.x = ReadDwordFromLDS(SharedIndex + 0 * LDS_SIZE);
-	V.y = ReadDwordFromLDS(SharedIndex + 1 * LDS_SIZE);
-	V.z = ReadDwordFromLDS(SharedIndex + 2 * LDS_SIZE);
-	V.w = ReadDwordFromLDS(SharedIndex + 3 * LDS_SIZE);
+	V.x = DW[0];
+	V.y = DW[1];
+	V.z = DW[2];
+	V.w = DW[3];
 }
 
 // float
-void WriteVectorToLDS(const uint SharedIndex, float V)
+void PackVectorToDwords(float V, out uint DW[1])
 {
-	WriteVectorToLDS(SharedIndex, asuint(V));
+	DW[0] = asuint(V);
 }
 
-void WriteVectorToLDS(const uint SharedIndex, float2 V)
+void PackVectorToDwords(float2 V, out uint DW[2])
 {
-	WriteVectorToLDS(SharedIndex, uint2(asuint(V.x), asuint(V.y)));
+	DW[0] = asuint(V.x);
+	DW[1] = asuint(V.y);
 }
 
-void WriteVectorToLDS(const uint SharedIndex, float3 V)
+void PackVectorToDwords(float3 V, out uint DW[3])
 {
-	WriteVectorToLDS(SharedIndex, uint3(asuint(V.x), asuint(V.y), asuint(V.z)));
+	DW[0] = asuint(V.x);
+	DW[1] = asuint(V.y);
+	DW[2] = asuint(V.z);
 }
 
-void WriteVectorToLDS(const uint SharedIndex, float4 V)
+void PackVectorToDwords(float4 V, out uint DW[4])
 {
-	WriteVectorToLDS(SharedIndex, uint4(asuint(V.x), asuint(V.y), asuint(V.z), asuint(V.w)));
+	DW[0] = asuint(V.x);
+	DW[1] = asuint(V.y);
+	DW[2] = asuint(V.z);
+	DW[3] = asuint(V.w);
 }
 
-void ReadVectorFromLDS(const uint SharedIndex, out float V)
+void UnpackDwordsToVector(uint DW[1], out float V)
 {
-	uint UV;
-	ReadVectorFromLDS(SharedIndex, /* out */ UV);
-	V = asfloat(UV);
+	V   = asfloat(DW[0]);
 }
 
-void ReadVectorFromLDS(const uint SharedIndex, out float2 V)
+void UnpackDwordsToVector(uint DW[2], out float2 V)
 {
-	uint2 UV;
-	ReadVectorFromLDS(SharedIndex, /* out */ UV);
-	V.x = asfloat(UV.x);
-	V.y = asfloat(UV.y);
+	V.x = asfloat(DW[0]);
+	V.y = asfloat(DW[1]);
 }
 
-void ReadVectorFromLDS(const uint SharedIndex, out float3 V)
+void UnpackDwordsToVector(uint DW[3], out float3 V)
 {
-	uint3 UV;
-	ReadVectorFromLDS(SharedIndex, /* out */ UV);
-	V.x = asfloat(UV.x);
-	V.y = asfloat(UV.y);
-	V.z = asfloat(UV.z);
+	V.x = asfloat(DW[0]);
+	V.y = asfloat(DW[1]);
+	V.z = asfloat(DW[2]);
 }
 
-void ReadVectorFromLDS(const uint SharedIndex, out float4 V)
+void UnpackDwordsToVector(uint DW[4], out float4 V)
 {
-	uint4 UV;
-	ReadVectorFromLDS(SharedIndex, /* out */ UV);
-	V.x = asfloat(UV.x);
-	V.y = asfloat(UV.y);
-	V.z = asfloat(UV.z);
-	V.w = asfloat(UV.w);
+	V.x = asfloat(DW[0]);
+	V.y = asfloat(DW[1]);
+	V.z = asfloat(DW[2]);
+	V.w = asfloat(DW[3]);
 }
 
 #if PLATFORM_SUPPORTS_REAL_TYPES
@@ -459,174 +449,220 @@ void ReadVectorFromLDS(const uint SharedIndex, out float4 V)
 #if COMPILER_SUPPORT_UINT16_BITCAST
 
 // uint16_t
-void WriteVectorToLDS(const uint SharedIndex, uint16_t V)
+void PackVectorToDwords(uint16_t V, out uint DW[1])
 {
-	WriteVectorToLDS(SharedIndex, uint(V));
+	DW[0] = uint(V);
 }
 
-void WriteVectorToLDS(const uint SharedIndex, uint16_t2 V)
+void PackVectorToDwords(uint16_t2 V, out uint DW[1])
 {
-	WriteVectorToLDS(SharedIndex, bit_cast_uint(V.xy));
+	DW[0] = bit_cast_uint(V.xy);
 }
 
-void WriteVectorToLDS(const uint SharedIndex, uint16_t3 V)
+void PackVectorToDwords(uint16_t3 V, out uint DW[2])
 {
-	WriteVectorToLDS(SharedIndex, uint2(bit_cast_uint(V.xy), uint(V.z)));
+	DW[0] = bit_cast_uint(V.xy);
+	DW[1] = uint(V.z);
 }
 
-void WriteVectorToLDS(const uint SharedIndex, uint16_t4 V)
+void PackVectorToDwords(uint16_t4 V, out uint DW[2])
 {
-	WriteVectorToLDS(SharedIndex, uint2(bit_cast_uint(V.xy), bit_cast_uint(V.zw)));
+	DW[0] = bit_cast_uint(V.xy);
+	DW[1] = bit_cast_uint(V.zw);
 }
 
-void ReadVectorFromLDS(const uint SharedIndex, out uint16_t V)
+void UnpackDwordsToVector(uint DW[1], out uint16_t V)
 {
-	uint UV;
-	ReadVectorFromLDS(SharedIndex, /* out */ UV);
-	V = tsr_ushort(UV);
+	V = uint16_t(DW[0]);
 }
 
-void ReadVectorFromLDS(const uint SharedIndex, out uint16_t2 V)
+void UnpackDwordsToVector(uint DW[1], out uint16_t2 V)
 {
-	uint UV;
-	ReadVectorFromLDS(SharedIndex, /* out */ UV);
-	V.xy = bit_cast_uint16_t2(UV);
+	V.xy = bit_cast_uint16_t2(DW[0]);
 }
 
-void ReadVectorFromLDS(const uint SharedIndex, out uint16_t3 V)
+void UnpackDwordsToVector(uint DW[2], out uint16_t3 V)
 {
-	uint2 UV;
-	ReadVectorFromLDS(SharedIndex, /* out */ UV);
-	V.xy = bit_cast_uint16_t2(UV.x);
-	V.z = tsr_ushort(UV.y);
+	V.xy = bit_cast_uint16_t2(DW[0]);
+	V.z = uint16_t(DW[1]);
 }
 
-void ReadVectorFromLDS(const uint SharedIndex, out uint16_t4 V)
+void UnpackDwordsToVector(uint DW[2], out uint16_t4 V)
 {
-	uint2 UV;
-	ReadVectorFromLDS(SharedIndex, /* out */ UV);
-	V.xy = bit_cast_uint16_t2(UV.x);
-	V.zw = bit_cast_uint16_t2(UV.y);
+	V.xy = bit_cast_uint16_t2(DW[0]);
+	V.zw = bit_cast_uint16_t2(DW[1]);
 }
 
 #else // !COMPILER_SUPPORT_UINT16_BITCAST
 
 // uint16_t
-void WriteVectorToLDS(const uint SharedIndex, uint16_t V)
+void PackVectorToDwords(uint16_t V, out uint DW[1])
 {
-	WriteVectorToLDS(SharedIndex, uint(V));
+	DW[0] = uint(V);
 }
 
-void WriteVectorToLDS(const uint SharedIndex, uint16_t2 V)
+void PackVectorToDwords(uint16_t2 V, out uint DW[1])
 {
-	WriteVectorToLDS(SharedIndex, uint(V.x) | (uint(V.y) << 16u));
+	DW[0] = uint(V.x) | (uint(V.y) << 16u);
 }
 
-void WriteVectorToLDS(const uint SharedIndex, uint16_t3 V)
+void PackVectorToDwords(uint16_t3 V, out uint DW[2])
 {
-	WriteVectorToLDS(SharedIndex, uint2(uint(V.x) | (uint(V.y) << 16u), uint(V.z)));
+	DW[0] = uint(V.x) | (uint(V.y) << 16u);
+	DW[1] = uint(V.z);
 }
 
-void WriteVectorToLDS(const uint SharedIndex, uint16_t4 V)
+void PackVectorToDwords(uint16_t4 V, out uint DW[2])
 {
-	WriteVectorToLDS(SharedIndex, uint2(uint(V.x) | (uint(V.y) << 16u), uint(V.z) | (uint(V.w) << 16u)));
+	DW[0] = uint(V.x) | (uint(V.y) << 16u);
+	DW[1] = uint(V.z) | (uint(V.w) << 16u);
 }
 
-void ReadVectorFromLDS(const uint SharedIndex, out uint16_t V)
+void UnpackDwordsToVector(uint DW[1], out uint16_t V)
 {
-	uint UV;
-	ReadVectorFromLDS(SharedIndex, /* out */ UV);
-	V = tsr_ushort(UV);
+	V = uint16_t(DW[0]);
 }
 
-void ReadVectorFromLDS(const uint SharedIndex, out uint16_t2 V)
+void UnpackDwordsToVector(uint DW[1], out uint16_t2 V)
 {
-	uint UV;
-	ReadVectorFromLDS(SharedIndex, /* out */ UV);
-	V.x = tsr_ushort((UV >>  0u) & 0xFFFFu);
-	V.y = tsr_ushort((UV >> 16u) & 0xFFFFu);
+	V.x = tsr_ushort((DW[0] >>  0u) & 0xFFFFu);
+	V.y = tsr_ushort((DW[0] >> 16u) & 0xFFFFu);
 }
 
-void ReadVectorFromLDS(const uint SharedIndex, out uint16_t3 V)
+void UnpackDwordsToVector(uint DW[2], out uint16_t3 V)
 {
-	uint2 UV;
-	ReadVectorFromLDS(SharedIndex, /* out */ UV);
-	V.x = tsr_ushort((UV.x >>  0u) & 0xFFFFu);
-	V.y = tsr_ushort((UV.x >> 16u) & 0xFFFFu);
-	V.z = tsr_ushort((UV.y >>  0u) & 0xFFFFu);
+	V.x = tsr_ushort((DW[0] >>  0u) & 0xFFFFu);
+	V.y = tsr_ushort((DW[0] >> 16u) & 0xFFFFu);
+	V.z = tsr_ushort((DW[1] >>  0u) & 0xFFFFu);
 }
 
-void ReadVectorFromLDS(const uint SharedIndex, out uint16_t4 V)
+void UnpackDwordsToVector(uint DW[2], out uint16_t4 V)
 {
-	uint2 UV;
-	ReadVectorFromLDS(SharedIndex, /* out */ UV);
-	V.x = tsr_ushort((UV.x >>  0u) & 0xFFFFu);
-	V.y = tsr_ushort((UV.x >> 16u) & 0xFFFFu);
-	V.z = tsr_ushort((UV.y >>  0u) & 0xFFFFu);
-	V.w = tsr_ushort((UV.y >> 16u) & 0xFFFFu);
+	V.x = tsr_ushort((DW[0] >>  0u) & 0xFFFFu);
+	V.y = tsr_ushort((DW[0] >> 16u) & 0xFFFFu);
+	V.z = tsr_ushort((DW[1] >>  0u) & 0xFFFFu);
+	V.w = tsr_ushort((DW[1] >> 16u) & 0xFFFFu);
 }
 
 #endif // !COMPILER_SUPPORT_UINT16_BITCAST
 
 // half
-void WriteVectorToLDS(const uint SharedIndex, half V)
+void PackVectorToDwords(half V, out uint DW[1])
 {
-	WriteVectorToLDS(SharedIndex, asuint16(V));
+	PackVectorToDwords(asuint16(V), /* out */ DW);
 }
 
-void WriteVectorToLDS(const uint SharedIndex, half2 V)
+void PackVectorToDwords(half2 V, out uint DW[1])
 {
-	WriteVectorToLDS(SharedIndex, uint16_t2(asuint16(V.x), asuint16(V.y)));
+	PackVectorToDwords(uint16_t2(asuint16(V.x), asuint16(V.y)), /* out */ DW);
 }
 
-void WriteVectorToLDS(const uint SharedIndex, half3 V)
+void PackVectorToDwords(half3 V, out uint DW[2])
 {
-	WriteVectorToLDS(SharedIndex, uint16_t3(asuint16(V.x), asuint16(V.y), asuint16(V.z)));
+	PackVectorToDwords(uint16_t3(asuint16(V.x), asuint16(V.y), asuint16(V.z)), /* out */ DW);
 }
 
-void WriteVectorToLDS(const uint SharedIndex, half4 V)
+void PackVectorToDwords(half4 V, out uint DW[2])
 {
-	WriteVectorToLDS(SharedIndex, uint16_t4(asuint16(V.x), asuint16(V.y), asuint16(V.z), asuint16(V.w)));
+	PackVectorToDwords(uint16_t4(asuint16(V.x), asuint16(V.y), asuint16(V.z), asuint16(V.w)), /* out */ DW);
 }
 
-void ReadVectorFromLDS(const uint SharedIndex, out half V)
+void UnpackDwordsToVector(uint DW[1], out half V)
 {
 	uint16_t UV;
-	ReadVectorFromLDS(SharedIndex, /* out */ UV);
+	UnpackDwordsToVector(DW, /* out */ UV);
 	V = asfloat16(UV);
 }
 
-void ReadVectorFromLDS(const uint SharedIndex, out half2 V)
+void UnpackDwordsToVector(uint DW[1], out half2 V)
 {
 	uint16_t2 UV;
-	ReadVectorFromLDS(SharedIndex, /* out */ UV);
+	UnpackDwordsToVector(DW, /* out */ UV);
 	V.x = asfloat16(UV.x);
 	V.y = asfloat16(UV.y);
 }
 
-void ReadVectorFromLDS(const uint SharedIndex, out half3 V)
+void UnpackDwordsToVector(uint DW[2], out half3 V)
 {
 	uint16_t3 UV;
-	ReadVectorFromLDS(SharedIndex, /* out */ UV);
+	UnpackDwordsToVector(DW, /* out */ UV);
 	V.x = asfloat16(UV.x);
 	V.y = asfloat16(UV.y);
 	V.z = asfloat16(UV.z);
 }
 
-void ReadVectorFromLDS(const uint SharedIndex, out half4 V)
+void UnpackDwordsToVector(uint DW[2], out half4 V)
 {
 	uint16_t4 UV;
-	ReadVectorFromLDS(SharedIndex, /* out */ UV);
+	UnpackDwordsToVector(DW, /* out */ UV);
 	V.x = asfloat16(UV.x);
 	V.y = asfloat16(UV.y);
 	V.z = asfloat16(UV.z);
 	V.w = asfloat16(UV.w);
 }
 
+
 #endif // PLATFORM_SUPPORTS_REAL_TYPES
 
+
+
+
+//------------------------------------------------------- GROUP SHARED COMMUNICATION
+
+#if defined(LDS_SIZE)
+
+groupshared uint SharedData[LDS_SIZE * LDS_DWORD_COMPONENT_COUNT];
+
+void WriteDwordToLDS(const uint SharedIndex, uint V)
+{
+	SharedData[SharedIndex] = V;
+}
+
+uint ReadDwordFromLDS(const uint SharedIndex)
+{
+	return SharedData[SharedIndex];
+}
+
+void AtomicIncrementLDSDword(const uint SharedIndex, uint V)
+{
+	InterlockedAdd(/* inout */ SharedData[SharedIndex], V);
+}
+
+template<typename ScalarType, uint VectorSize>
+void WriteVectorToLDS(const uint SharedIndex, vector<ScalarType, VectorSize> V)
+{
+	uint DW[GetVectorDwordSize(ScalarType, VectorSize)];
+	PackVectorToDwords(V, /* out */ DW);
+	
+	UNROLL
+	for (uint i = 0; i < GetVectorDwordSize(ScalarType, VectorSize); i++)
+	{
+		SharedData[SharedIndex + i * LDS_SIZE] = DW[i];
+	}
+}
+
+template<typename ScalarType, uint VectorSize>
+void ReadVectorFromLDS(const uint SharedIndex, out vector<ScalarType, VectorSize> V)
+{
+	uint DW[GetVectorDwordSize(ScalarType, VectorSize)];
+	UNROLL
+	for (uint i = 0; i < GetVectorDwordSize(ScalarType, VectorSize); i++)
+	{
+		DW[i] = SharedData[SharedIndex + i * LDS_SIZE];
+	}
+	UnpackDwordsToVector(DW, /* out */ V);
+}
+
 #endif // defined(LDS_SIZE)
 
+//------------------------------------------------------- MANUAL LDS DESPILL
+
+#if defined(LDS_DESPILL_DWORD_COUNT)
+
+groupshared uint DespilledData[LDS_DESPILL_DWORD_COUNT];
+
+#endif // defined(LDS_DESPILL_DWORD_COUNT)
+
 
 //------------------------------------------------------- REGISTERS LAYOUTS
 
@@ -1247,6 +1283,59 @@ struct TLaneVector2D
 	}
 	#endif
 	
+	#if defined(LDS_DESPILL_DWORD_COUNT)
+	
+	/** Returns the LDS spill size in dwords */
+	static uint GetLDSSpillDwordSize()
+	{
+		return RegisterRowCount * GetVectorDwordSize(ScalarType, RegisterRowSize);
+	}
+	
+	/** Store to LDS spill to save VGPR. */
+	CALL_SITE_DEBUGLOC
+	void StoreLDSSpill(const uint DespillDwordOffset)
+	{
+		UNROLL
+		for (uint RegisterRowIndex = 0; RegisterRowIndex < RegisterRowCount; RegisterRowIndex++)
+		{
+			vector<ScalarType, RegisterRowSize> ResgiterRow = Registers.GetRegisterRow(RegisterRowIndex);
+			
+			uint ResgiterRowDW[GetVectorDwordSize(ScalarType, RegisterRowSize)];
+			PackVectorToDwords(ResgiterRow, /* out */ ResgiterRowDW);
+	
+			UNROLL
+			for (uint i = 0; i < GetVectorDwordSize(ScalarType, RegisterRowSize); i++)
+			{
+				DespilledData[GGroupThreadIndex + DespillDwordOffset + (RegisterRowIndex * GetVectorDwordSize(ScalarType, RegisterRowSize) + i) * LDS_DESPILL_THREAD_COUNT] = ResgiterRowDW[i];
+			}
+		}
+	}
+	
+	/** Load from LDS spill to save VGPR. */
+	CALL_SITE_DEBUGLOC
+	static TLaneVector2D<ScalarType, VectorSize, SimdSizeX, SimdSizeY> LoadLDSSpill(const uint DespillDwordOffset)
+	{
+		TLaneVector2D<ScalarType, VectorSize, SimdSizeX, SimdSizeY> Vector;
+		UNROLL
+		for (uint RegisterRowIndex = 0; RegisterRowIndex < RegisterRowCount; RegisterRowIndex++)
+		{
+			uint ResgiterRowDW[GetVectorDwordSize(ScalarType, RegisterRowSize)];
+			UNROLL
+			for (uint i = 0; i < GetVectorDwordSize(ScalarType, RegisterRowSize); i++)
+			{
+				ResgiterRowDW[i] = DespilledData[GGroupThreadIndex + DespillDwordOffset + (RegisterRowIndex * GetVectorDwordSize(ScalarType, RegisterRowSize) + i) * LDS_DESPILL_THREAD_COUNT];
+			}
+			
+			vector<ScalarType, RegisterRowSize> ResgiterRow;
+			UnpackDwordsToVector(ResgiterRowDW, /* out */ ResgiterRow);
+	
+			Vector.Registers.SetRegisterRow(RegisterRowIndex, ResgiterRow);
+		}
+		return Vector;
+	}
+	
+	#endif // defined(LDS_DESPILL_DWORD_COUNT)
+	
 	
 	// ------------- binary operator +
 	CALL_SITE_DEBUGLOC
