Fixing bindless resource index prefixes after Constant Buffer object changes
- It was generating UniformBuffer_BindlessResource_Member instead of BindlessResource_UniformBuffer_Member.
- The prefix needs to be applied to the full symbol name for correctly mapping the index against its resource in reflection.

diff --git a/Engine/Shaders/Public/BindlessResources.ush b/Engine/Shaders/Public/BindlessResources.ush
--- a/Engine/Shaders/Public/BindlessResources.ush
+++ b/Engine/Shaders/Public/BindlessResources.ush
@@ -17,7 +17,7 @@
 
 #elif ENABLE_BINDLESS_RESOURCES && PLATFORM_SUPPORTS_BINDLESS
 	#define UB_RESOURCE_MEMBER_RESOURCE(Type, UBName, Name) \
-		Type GetBindlessResource##UBName##_##Name() { return GetResourceFromHeap(Type, UB_CB_MEMBER_ACCESS(UBName, BindlessResource_##Name)); } \
+		Type GetBindlessResource##UBName##_##Name() { return GetResourceFromHeap(Type, UB_CB_PREFIXED_MEMBER_ACCESS(UBName, BindlessResource_, Name)); } \
 		static const Type UBName##_##Name = GetBindlessResource##UBName##_##Name();
 #else
 	#define UB_RESOURCE_MEMBER_RESOURCE(Type, UBName, Name) Type UBName##_##Name
@@ -28,7 +28,7 @@
 
 #elif ENABLE_BINDLESS_SAMPLERS && PLATFORM_SUPPORTS_BINDLESS
 	#define UB_RESOURCE_MEMBER_SAMPLER(Type, UBName, Name) \
-		Type GetBindlessSampler##UBName##_##Name() { return GetSamplerFromHeap(Type, UB_CB_MEMBER_ACCESS(UBName, BindlessSampler_##Name)); } \
+		Type GetBindlessSampler##UBName##_##Name() { return GetSamplerFromHeap(Type, UB_CB_PREFIXED_MEMBER_ACCESS(UBName, BindlessSampler_, Name)); } \
 		static const Type UBName##_##Name = GetBindlessSampler##UBName##_##Name();
 #else
 	#define UB_RESOURCE_MEMBER_SAMPLER(Type, UBName, Name) Type UBName##_##Name
diff --git a/Engine/Shaders/Public/Platform.ush b/Engine/Shaders/Public/Platform.ush
--- a/Engine/Shaders/Public/Platform.ush
+++ b/Engine/Shaders/Public/Platform.ush
@@ -457,18 +457,22 @@
 #define PLATFORM_MAX_UNIFORM_BUFFER_RANGE_FLOAT4 (PLATFORM_MAX_UNIFORM_BUFFER_RANGE/16u)
 
 #if PLATFORM_SUPPORTS_CONSTANTBUFFER_OBJECT
-	#define UB_CB_DEFINITION_START(UBName            ) struct F##UBName##Constants {
-	#define UB_CB_MEMBER_NAME(     UBName, MemberName) MemberName
-	#define UB_CB_DEFINITION_END(  UBName            ) }; ConstantBuffer<F##UBName##Constants> UBName;
-	#define UB_CB_MEMBER_ACCESS(   UBName, MemberName) UBName##.##MemberName
+	#define UB_CB_DEFINITION_START(      UBName                    ) struct F##UBName##Constants {
+	#define UB_CB_MEMBER_NAME(           UBName,         MemberName) MemberName
+	#define UB_CB_PREFIXED_MEMBER_NAME(  UBName, Prefix, MemberName) Prefix##MemberName
+	#define UB_CB_DEFINITION_END(        UBName                    ) }; ConstantBuffer<F##UBName##Constants> UBName;
+	#define UB_CB_MEMBER_ACCESS(         UBName,         MemberName) UBName##.##MemberName
+	#define UB_CB_PREFIXED_MEMBER_ACCESS(UBName, Prefix, MemberName) UBName##.##Prefix##MemberName
 
 	#define UB_DECL_PARAMETER(UBName, StructName, GlobalName) UBName##.##StructName = UBName##.##GlobalName
 	#define UB_DECL_RESOURCE(UBName, StructName, GlobalName)  UBName##.##StructName = UBName##_##GlobalName
 #else
-	#define UB_CB_DEFINITION_START(UBName            ) cbuffer UBName {
-	#define UB_CB_MEMBER_NAME(     UBName, MemberName) UBName##_##MemberName
-	#define UB_CB_DEFINITION_END(  UBName            ) }
-	#define UB_CB_MEMBER_ACCESS(   UBName, MemberName) UBName##_##MemberName
+	#define UB_CB_DEFINITION_START(      UBName                    ) cbuffer UBName {
+	#define UB_CB_MEMBER_NAME(           UBName,         MemberName) UBName##_##MemberName
+	#define UB_CB_PREFIXED_MEMBER_NAME(  UBName, Prefix, MemberName) Prefix##UBName##_##MemberName
+	#define UB_CB_DEFINITION_END(        UBName                    ) }
+	#define UB_CB_MEMBER_ACCESS(         UBName,         MemberName) UBName##_##MemberName
+	#define UB_CB_PREFIXED_MEMBER_ACCESS(UBName, Prefix, MemberName) Prefix##UBName##_##MemberName
 
 	#define UB_DECL_PARAMETER(UBName, StructName, GlobalName) UBName##.##StructName = UBName##_##GlobalName
 	#define UB_DECL_RESOURCE(UBName, StructName, GlobalName)  UBName##.##StructName = UBName##_##GlobalName
diff --git a/Engine/Shaders/Public/Platform/Vulkan/VulkanCommon.ush b/Engine/Shaders/Public/Platform/Vulkan/VulkanCommon.ush
--- a/Engine/Shaders/Public/Platform/Vulkan/VulkanCommon.ush
+++ b/Engine/Shaders/Public/Platform/Vulkan/VulkanCommon.ush
@@ -46,7 +46,7 @@ void ImageInterlockedMaxUInt64(RWTexture2D<UlongType> Dest, uint2 Coord, UlongTy
 		typedef Type SafeType##UBName##_##Name; \
 		SafeType##UBName##_##Name VULKAN_RESOURCE_HEAP(SafeType##UBName##_##Name)[]; \
 		Type GetBindlessResource##UBName##_##Name() { \
-			return VULKAN_RESOURCE_HEAP(SafeType##UBName##_##Name)[UB_CB_MEMBER_ACCESS(UBName, BindlessResource_##Name)]; \
+			return VULKAN_RESOURCE_HEAP(SafeType##UBName##_##Name)[UB_CB_PREFIXED_MEMBER_ACCESS(UBName, BindlessResource_, Name)]; \
 		} \
 		static const Type UBName##_##Name = GetBindlessResource##UBName##_##Name();
 
@@ -54,7 +54,7 @@ void ImageInterlockedMaxUInt64(RWTexture2D<UlongType> Dest, uint2 Coord, UlongTy
 		typedef Type SafeType##UBName##_##Name; \
 		SafeType##UBName##_##Name VULKAN_SAMPLER_HEAP(SafeType##UBName##_##Name)[]; \
 		Type GetBindlessSampler##UBName##_##Name() { \
-			return VULKAN_SAMPLER_HEAP(SafeType##UBName##_##Name)[UB_CB_MEMBER_ACCESS(UBName, BindlessSampler_##Name)]; \
+			return VULKAN_SAMPLER_HEAP(SafeType##UBName##_##Name)[UB_CB_PREFIXED_MEMBER_ACCESS(UBName, BindlessSampler_, Name)]; \
 		} \
 		static const Type UBName##_##Name = GetBindlessSampler##UBName##_##Name();
 #endif
diff --git a/Engine/Source/Runtime/RenderCore/Private/ShaderParameters.cpp b/Engine/Source/Runtime/RenderCore/Private/ShaderParameters.cpp
--- a/Engine/Source/Runtime/RenderCore/Private/ShaderParameters.cpp
+++ b/Engine/Source/Runtime/RenderCore/Private/ShaderParameters.cpp
@@ -241,16 +241,16 @@ static void CreateHLSLUniformBufferStructMembersDeclaration(
 			HLSLBaseOffset = AbsoluteMemberOffset + HLSLMemberSize;
 
 			// Generate the member declaration.
+			const TCHAR* MemberPrefix = (Member.GetBaseType() == UBMT_SAMPLER ? FShaderParameterParser::kBindlessSamplerPrefix : FShaderParameterParser::kBindlessResourcePrefix);
 
 			TStringBuilder<256> ParameterName;
-			ParameterName << (Member.GetBaseType() == UBMT_SAMPLER ? FShaderParameterParser::kBindlessSamplerPrefix : FShaderParameterParser::kBindlessResourcePrefix);
 			if (!GlobalPrefix.IsEmpty())
 			{
 				ParameterName << GlobalPrefix << TEXT("_");
 			}
 			ParameterName << Member.GetName();
 
-			Decl.ConstantBufferMembers.Appendf(TEXT("\tUB_UINT() UB_CB_MEMBER_NAME(%s, %s);\r\n"), *UniformBufferName, *ParameterName);
+			Decl.ConstantBufferMembers.Appendf(TEXT("\tUB_UINT() UB_CB_PREFIXED_MEMBER_NAME(%s, %s, %s);\r\n"), *UniformBufferName, MemberPrefix, *ParameterName);
 		}
 		else
 		{
