Bump ShaderVersion.ush GUID as follow up on CL 29088141.

#rnx

diff --git a/Engine/Shaders/Public/ShaderVersion.ush b/Engine/Shaders/Public/ShaderVersion.ush
--- a/Engine/Shaders/Public/ShaderVersion.ush
+++ b/Engine/Shaders/Public/ShaderVersion.ush
@@ -2,4 +2,4 @@
 // Each time the console command is executed it generates a new GUID. As this file is included
 // in Platform.ush (which should be included in any shader) it allows to invalidate the shader DDC.
 // 
-#pragma message("UESHADERMETADATA_VERSION 15A33BF2-B4B6-4932-BF5A-6BD3D2689F97")
+#pragma message("UESHADERMETADATA_VERSION 78C3F233-CBA3-4DB9-8BF5-F1FE3E490DD3")
