Fixes TSR on consoles preview in editor

diff --git a/Engine/Shaders/Private/TemporalSuperResolution/TSRCommon.ush b/Engine/Shaders/Private/TemporalSuperResolution/TSRCommon.ush
--- a/Engine/Shaders/Private/TemporalSuperResolution/TSRCommon.ush
+++ b/Engine/Shaders/Private/TemporalSuperResolution/TSRCommon.ush
@@ -1186,12 +1186,12 @@ FCatmullRomSamples GetBicubic2DCatmullRomSamples_Stubbe(float2 UV, float2 Size,
 	Samples.UV[2] = cuv + uv;
 	Samples.UV[3] = cuv + tsr_half2(                     uv.x,                -InvSize.y);
 	Samples.UV[4] = cuv + tsr_half2(                     uv.x, tsr_half(2.0) * InvSize.y);
-
-	Samples.Weight[0] = w0.x;
-	Samples.Weight[1] = w3.x;
-	Samples.Weight[2] = tsr_half(1.0) - w0.x - w3.x - w0.y - w3.y;
-	Samples.Weight[3] = w0.y;
-	Samples.Weight[4] = w3.y;
+	
+	Samples.Weight[0] = half(w0.x);
+	Samples.Weight[1] = half(w3.x);
+	Samples.Weight[2] = half(tsr_half(1.0) - w0.x - w3.x - w0.y - w3.y);
+	Samples.Weight[3] = half(w0.y);
+	Samples.Weight[4] = half(w3.y);
 
 	Samples.UVDir[0] = int2(-1,  0);
 	Samples.UVDir[1] = int2( 1,  0);
diff --git a/Engine/Shaders/Private/TemporalSuperResolution/TSRDilateVelocity.usf b/Engine/Shaders/Private/TemporalSuperResolution/TSRDilateVelocity.usf
--- a/Engine/Shaders/Private/TemporalSuperResolution/TSRDilateVelocity.usf
+++ b/Engine/Shaders/Private/TemporalSuperResolution/TSRDilateVelocity.usf
@@ -288,7 +288,9 @@ void MainCS(
 		float3 ScreenVelocity1 = ComputeStaticVelocity(ScreenPos, ClosestDeviceZ1);
 		float3 ScreenVelocity2 = ComputeStaticVelocity(ScreenPos, ClosestDeviceZ2);
 		
-		HasPixelAnimation = DecodeHasPixelAnimationFromVelocityTexture(EncodedVelocity1);
+		#if VELOCITY_ENCODE_HAS_PIXEL_ANIMATION
+			HasPixelAnimation = DecodeHasPixelAnimationFromVelocityTexture(EncodedVelocity1);
+		#endif
 
 		#if CONFIG_USE_SCENE_VELOCITY
 		{
@@ -315,7 +317,9 @@ void MainCS(
 		
 		float4 EncodedVelocity = SceneVelocityTexture[uint2(InputBufferUV * InputInfo_Extent) + VelocityPixelOffset];
 
-		HasPixelAnimation = DecodeHasPixelAnimationFromVelocityTexture(EncodedVelocity);
+		#if VELOCITY_ENCODE_HAS_PIXEL_ANIMATION
+			HasPixelAnimation = DecodeHasPixelAnimationFromVelocityTexture(EncodedVelocity);
+		#endif
 
 		#if CONFIG_USE_SCENE_VELOCITY
 		{
diff --git a/Engine/Shaders/Public/LaneVectorization.ush b/Engine/Shaders/Public/LaneVectorization.ush
--- a/Engine/Shaders/Public/LaneVectorization.ush
+++ b/Engine/Shaders/Public/LaneVectorization.ush
@@ -525,23 +525,23 @@ void UnpackDwordsToVector(uint DW[1], out uint16_t V)
 
 void UnpackDwordsToVector(uint DW[1], out uint16_t2 V)
 {
-	V.x = tsr_ushort((DW[0] >>  0u) & 0xFFFFu);
-	V.y = tsr_ushort((DW[0] >> 16u) & 0xFFFFu);
+	V.x = uint16_t((DW[0] >>  0u) & 0xFFFFu);
+	V.y = uint16_t((DW[0] >> 16u) & 0xFFFFu);
 }
 
 void UnpackDwordsToVector(uint DW[2], out uint16_t3 V)
 {
-	V.x = tsr_ushort((DW[0] >>  0u) & 0xFFFFu);
-	V.y = tsr_ushort((DW[0] >> 16u) & 0xFFFFu);
-	V.z = tsr_ushort((DW[1] >>  0u) & 0xFFFFu);
+	V.x = uint16_t((DW[0] >>  0u) & 0xFFFFu);
+	V.y = uint16_t((DW[0] >> 16u) & 0xFFFFu);
+	V.z = uint16_t((DW[1] >>  0u) & 0xFFFFu);
 }
 
 void UnpackDwordsToVector(uint DW[2], out uint16_t4 V)
 {
-	V.x = tsr_ushort((DW[0] >>  0u) & 0xFFFFu);
-	V.y = tsr_ushort((DW[0] >> 16u) & 0xFFFFu);
-	V.z = tsr_ushort((DW[1] >>  0u) & 0xFFFFu);
-	V.w = tsr_ushort((DW[1] >> 16u) & 0xFFFFu);
+	V.x = uint16_t((DW[0] >>  0u) & 0xFFFFu);
+	V.y = uint16_t((DW[0] >> 16u) & 0xFFFFu);
+	V.z = uint16_t((DW[1] >>  0u) & 0xFFFFu);
+	V.w = uint16_t((DW[1] >> 16u) & 0xFFFFu);
 }
 
 #endif // !COMPILER_SUPPORT_UINT16_BITCAST
diff --git a/Engine/Source/Runtime/Renderer/Private/PostProcess/TemporalSuperResolution.cpp b/Engine/Source/Runtime/Renderer/Private/PostProcess/TemporalSuperResolution.cpp
--- a/Engine/Source/Runtime/Renderer/Private/PostProcess/TemporalSuperResolution.cpp
+++ b/Engine/Source/Runtime/Renderer/Private/PostProcess/TemporalSuperResolution.cpp
@@ -403,6 +403,12 @@ public:
 
 	static bool ShouldCompile32or16BitPermutation(EShaderPlatform Platform, bool bIs16BitVALUPermutation)
 	{
+		// Always compile 32bit ops on preview platform 
+		if (FDataDrivenShaderPlatformInfo::GetIsPreviewPlatform(Platform) && !bIs16BitVALUPermutation)
+		{
+			return true;
+		}
+
 		const ERHIFeatureSupport Support = FTSRShader::Supports16BitVALU(Platform);
 
 		if (Support == ERHIFeatureSupport::RuntimeGuaranteed)
@@ -428,6 +434,12 @@ public:
 
 	static bool SupportsLDS(EShaderPlatform Platform)
 	{
+		// Always support LDS on preview platform 
+		if (FDataDrivenShaderPlatformInfo::GetIsPreviewPlatform(Platform))
+		{
+			return true;
+		}
+
 		// Always support LDS if wave ops are not guarenteed
 		if (SupportsWaveOps(Platform) != ERHIFeatureSupport::RuntimeGuaranteed)
 		{
