Partial revert of 29803999 to fix shader cook issues on linux/ios

diff --git a/Engine/Plugins/FX/Niagara/Shaders/Private/NiagaraShaderVersion.ush b/Engine/Plugins/FX/Niagara/Shaders/Private/NiagaraShaderVersion.ush
--- a/Engine/Plugins/FX/Niagara/Shaders/Private/NiagaraShaderVersion.ush
+++ b/Engine/Plugins/FX/Niagara/Shaders/Private/NiagaraShaderVersion.ush
@@ -3,4 +3,4 @@
 // in the DDC key, it will automatically invalidate.
 // 
 // If you are merging streams and there is a conflict with this GUID you should make a new GUID rather than taking one or the other.
-#pragma message("UESHADERMETADATA_VERSION C0DE8DEB-9CBC-4800-8159-50B678CAC3F1")
\ No newline at end of file
+#pragma message("UESHADERMETADATA_VERSION 8AC67753-9EBF-41A4-BB15-5954EBED220C")
\ No newline at end of file
diff --git a/Engine/Shaders/Public/Platform.ush b/Engine/Shaders/Public/Platform.ush
--- a/Engine/Shaders/Public/Platform.ush
+++ b/Engine/Shaders/Public/Platform.ush
@@ -709,11 +709,7 @@ void ShaderYield()
 #endif
 
 #ifndef INVARIANT
-precise float MakePrecise(float v) { precise float pv = v; return pv; }
-precise float2 MakePrecise(float2 v) { precise float2 pv = v; return pv; }
-precise float3 MakePrecise(float3 v) { precise float3 pv = v; return pv; }
-precise float4 MakePrecise(float4 v) { precise float4 pv = v; return pv; }
-#define INVARIANT(X) MakePrecise(X)
+#define INVARIANT(X) (X)
 #endif
 
 #ifndef INVARIANT_OUTPUT
diff --git a/Engine/Shaders/Public/Platform/D3D/D3DCommon.ush b/Engine/Shaders/Public/Platform/D3D/D3DCommon.ush
--- a/Engine/Shaders/Public/Platform/D3D/D3DCommon.ush
+++ b/Engine/Shaders/Public/Platform/D3D/D3DCommon.ush
@@ -14,6 +14,12 @@ D3DCommon.ush: Common defines for the D3D shader compiler
 	#define WAVESIZE(N) [WaveSize(N)]
 #endif
 
+precise float MakePrecise(float v) { precise float pv = v; return pv; }
+precise float2 MakePrecise(float2 v) { precise float2 pv = v; return pv; }
+precise float3 MakePrecise(float3 v) { precise float3 pv = v; return pv; }
+precise float4 MakePrecise(float4 v) { precise float4 pv = v; return pv; }
+#define INVARIANT(X) MakePrecise(X)
+
 #if PLATFORM_SUPPORTS_STATIC_SAMPLERS
 
 // Static sampler table must match D3D12RootSignature.cpp
diff --git a/Engine/Shaders/Public/ShaderVersion.ush b/Engine/Shaders/Public/ShaderVersion.ush
--- a/Engine/Shaders/Public/ShaderVersion.ush
+++ b/Engine/Shaders/Public/ShaderVersion.ush
@@ -2,4 +2,4 @@
 // Each time the console command is executed it generates a new GUID. As this file is included
 // in Platform.ush (which should be included in any shader) it allows to invalidate the shader DDC.
 // 
-#pragma message("UESHADERMETADATA_VERSION CDF99C95-0301-4093-A42E-CBEDCF36071A")
+#pragma message("UESHADERMETADATA_VERSION A8F68E0C-FEF9-4330-8B8C-C1F8C4FFA445")
