Move INVARIANT definition from the D3D header to Platform. This applies the definition to all platforms except the ones that override it. Originally submitted as CL 29803999, later reverted until a HLSLcc problem was fixed in CL 29886548, and now resubmitted.

diff --git a/Engine/Shaders/Public/Platform.ush b/Engine/Shaders/Public/Platform.ush
--- a/Engine/Shaders/Public/Platform.ush
+++ b/Engine/Shaders/Public/Platform.ush
@@ -709,7 +709,11 @@ void ShaderYield()
 #endif
 
 #ifndef INVARIANT
-#define INVARIANT(X) (X)
+precise float MakePrecise(float v) { precise float pv = v; return pv; }
+precise float2 MakePrecise(float2 v) { precise float2 pv = v; return pv; }
+precise float3 MakePrecise(float3 v) { precise float3 pv = v; return pv; }
+precise float4 MakePrecise(float4 v) { precise float4 pv = v; return pv; }
+#define INVARIANT(X) MakePrecise(X)
 #endif
 
 #ifndef INVARIANT_OUTPUT
diff --git a/Engine/Shaders/Public/Platform/D3D/D3DCommon.ush b/Engine/Shaders/Public/Platform/D3D/D3DCommon.ush
--- a/Engine/Shaders/Public/Platform/D3D/D3DCommon.ush
+++ b/Engine/Shaders/Public/Platform/D3D/D3DCommon.ush
@@ -14,12 +14,6 @@ D3DCommon.ush: Common defines for the D3D shader compiler
 	#define WAVESIZE(N) [WaveSize(N)]
 #endif
 
-precise float MakePrecise(float v) { precise float pv = v; return pv; }
-precise float2 MakePrecise(float2 v) { precise float2 pv = v; return pv; }
-precise float3 MakePrecise(float3 v) { precise float3 pv = v; return pv; }
-precise float4 MakePrecise(float4 v) { precise float4 pv = v; return pv; }
-#define INVARIANT(X) MakePrecise(X)
-
 #if PLATFORM_SUPPORTS_STATIC_SAMPLERS
 
 // Static sampler table must match D3D12RootSignature.cpp
