Add PLATFORM_SUPPORTS_FMA and FMA for platform-specific fused-muliply add support.

diff --git a/Engine/Shaders/Public/Platform.ush b/Engine/Shaders/Public/Platform.ush
--- a/Engine/Shaders/Public/Platform.ush
+++ b/Engine/Shaders/Public/Platform.ush
@@ -720,6 +720,11 @@ precise float4 MakePrecise(in precise float4 v) { precise float4 pv = v; return
 #define INVARIANT_OUTPUT
 #endif
 
+#ifndef PLATFORM_SUPPORTS_FMA
+#define PLATFORM_SUPPORTS_FMA 1
+#define FMA(a, b, c) mad(a, b, c)
+#endif
+
 #ifndef ENABLE_RE_Z
 #define ENABLE_RE_Z
 #endif
