Improvements to RDNA ISA logic

diff --git a/Engine/Shaders/Public/ShaderVersion.ush b/Engine/Shaders/Public/ShaderVersion.ush
--- a/Engine/Shaders/Public/ShaderVersion.ush
+++ b/Engine/Shaders/Public/ShaderVersion.ush
@@ -3,4 +3,4 @@
 // Each time the console command is executed it generates a new GUID. As this file is included
 // in Platform.ush (which should be included in any shader) it allows to invalidate the shader DDC.
 // 
-#pragma message("UESHADERMETADATA_VERSION F632B8CA-8FAA-4BBC-B489-1618212BA5F0")
+#pragma message("UESHADERMETADATA_VERSION DDDDAA12-1111-46A4-8DD7-0FEBBBF2914")
