#UE Bump shader version after backout (31415008) to get newly generated shaders

diff --git a/Engine/Shaders/Public/ShaderVersion.ush b/Engine/Shaders/Public/ShaderVersion.ush
--- a/Engine/Shaders/Public/ShaderVersion.ush
+++ b/Engine/Shaders/Public/ShaderVersion.ush
@@ -3,4 +3,4 @@
 // Each time the console command is executed it generates a new GUID. As this file is included
 // in Platform.ush (which should be included in any shader) it allows to invalidate the shader DDC.
 // 
-#pragma message("UESHADERMETADATA_VERSION DC1DB8B4-F092-46B2-8B29-0A59A08B8D93")
+#pragma message("UESHADERMETADATA_VERSION 3F6C865B-CABB-4223-A659-0102915647FF")
