Invalidating Metal shaders due to issue on build machine

diff --git a/Engine/Shaders/Public/Platform/Metal/MetalCommon.ush b/Engine/Shaders/Public/Platform/Metal/MetalCommon.ush
--- a/Engine/Shaders/Public/Platform/Metal/MetalCommon.ush
+++ b/Engine/Shaders/Public/Platform/Metal/MetalCommon.ush
@@ -7,7 +7,7 @@
 #pragma once
 
 // Update this GUID to invalidate shader recompilation for only Metal shaders
-#pragma message("UESHADERMETADATA_VERSION 059503F5-57EB-4E3E-96DB-63F381FC9FEC")
+#pragma message("UESHADERMETADATA_VERSION E4FA73AD-5066-4067-8B47-814FDC0D6332")
 
 #if VERTEXSHADER && HULLSHADER
 #ifndef TESSELLATIONSHADER
