Add interlocked add uint64 support to texture 2d.

diff --git a/Engine/Shaders/Public/Platform/D3D/D3DCommon.ush b/Engine/Shaders/Public/Platform/D3D/D3DCommon.ush
--- a/Engine/Shaders/Public/Platform/D3D/D3DCommon.ush
+++ b/Engine/Shaders/Public/Platform/D3D/D3DCommon.ush
@@ -132,6 +132,11 @@ void UEReportAssertWithPayload(uint AssertID, uint4 Payload)
 		{
 			InterlockedMax(Dest[Coord], Value);
 		}
+
+		void ImageInterlockedAddUInt64(RWTexture2D<UlongType> Dest, uint2 Coord, UlongType Value)
+		{
+			InterlockedAdd(Dest[Coord], Value);
+		}
 	#endif
 #endif
 
diff --git a/Engine/Shaders/Public/Platform/Metal/MetalCommon.ush b/Engine/Shaders/Public/Platform/Metal/MetalCommon.ush
--- a/Engine/Shaders/Public/Platform/Metal/MetalCommon.ush
+++ b/Engine/Shaders/Public/Platform/Metal/MetalCommon.ush
@@ -89,6 +89,12 @@
         {
             InterlockedMax(Dest[Coord], Value);
         }
+
+		void ImageInterlockedAddUInt64(RWTexture2D<UlongType> Dest, uint2 Coord, UlongType Value)
+		{
+			InterlockedAdd(Dest[Coord], Value);
+		}
+
     #endif
     
     #if PLATFORM_SUPPORTS_BINDLESS
diff --git a/Engine/Shaders/Public/Platform/Vulkan/VulkanCommon.ush b/Engine/Shaders/Public/Platform/Vulkan/VulkanCommon.ush
--- a/Engine/Shaders/Public/Platform/Vulkan/VulkanCommon.ush
+++ b/Engine/Shaders/Public/Platform/Vulkan/VulkanCommon.ush
@@ -33,6 +33,11 @@ void ImageInterlockedMaxUInt64(RWTexture2D<UlongType> Dest, uint2 Coord, UlongTy
 {
 	InterlockedMax(Dest[Coord], Value);
 }
+
+void ImageInterlockedAddUInt64(RWTexture2D<UlongType> Dest, uint2 Coord, UlongType Value)
+{
+	InterlockedAdd(Dest[Coord], Value);
+}
 #endif
 
 #if PLATFORM_SUPPORTS_BINDLESS
