[shaders] version bump to clear corrupt data from DDC

#rnx

diff --git a/Engine/Shaders/Public/ShaderVersion.ush b/Engine/Shaders/Public/ShaderVersion.ush
--- a/Engine/Shaders/Public/ShaderVersion.ush
+++ b/Engine/Shaders/Public/ShaderVersion.ush
@@ -3,4 +3,4 @@
 // Each time the console command is executed it generates a new GUID. As this file is included
 // in Platform.ush (which should be included in any shader) it allows to invalidate the shader DDC.
 // 
-#pragma message("UESHADERMETADATA_VERSION 012B1460-472D-DC48-49C9-7B9537F03DFF")
\ No newline at end of file
+#pragma message("UESHADERMETADATA_VERSION 01BB4755-EA53-4DA2-BF4E-C29501DEEDE1")
