Bump GLCommon

diff --git a/Engine/Shaders/Public/Platform/GL/GLCommon.ush b/Engine/Shaders/Public/Platform/GL/GLCommon.ush
--- a/Engine/Shaders/Public/Platform/GL/GLCommon.ush
+++ b/Engine/Shaders/Public/Platform/GL/GLCommon.ush
@@ -7,6 +7,6 @@
 #pragma once
 
 // Update this GUID to invalidate shader recompilation for only GL shaders
-#pragma message("UESHADERMETADATA_VERSION 3FB1EB49-0442-40E1-9F1B-0371F7ABB4BB")
+#pragma message("UESHADERMETADATA_VERSION 05706E20-1BE5-45C7-868F-CEB8BCE25DB8")
 
 
