Bump GUID in ShaderVersion.ush to resolve stale shader DDC.

#rnx

diff --git a/Engine/Shaders/Public/ShaderVersion.ush b/Engine/Shaders/Public/ShaderVersion.ush
--- a/Engine/Shaders/Public/ShaderVersion.ush
+++ b/Engine/Shaders/Public/ShaderVersion.ush
@@ -3,4 +3,4 @@
 // Each time the console command is executed it generates a new GUID. As this file is included
 // in Platform.ush (which should be included in any shader) it allows to invalidate the shader DDC.
 // 
-#pragma message("UESHADERMETADATA_VERSION 56A27DD9-F88A-4CF4-85FC-034CC38B7CEB")
+#pragma message("UESHADERMETADATA_VERSION 6D4B37CE-368C-4654-BCF8-BF58084A1941")
