Adding the SafeType typedefs for uniform buffer members to avoid rewriter issues with bindless enabled.

diff --git a/Engine/Shaders/Public/BindlessResources.ush b/Engine/Shaders/Public/BindlessResources.ush
--- a/Engine/Shaders/Public/BindlessResources.ush
+++ b/Engine/Shaders/Public/BindlessResources.ush
@@ -17,12 +17,14 @@
 
 #elif ENABLE_BINDLESS_RESOURCES && PLATFORM_SUPPORTS_BINDLESS
 	#define UB_RESOURCE_MEMBER_SRV(Type, UBName, Name) \
-		Type GetBindlessResource##UBName##_##Name() { return GetResourceFromHeap(Type, UB_CB_PREFIXED_MEMBER_ACCESS(UBName, BindlessSRV_, Name)); } \
-		static const Type UBName##_##Name = GetBindlessResource##UBName##_##Name();
+		typedef Type SafeType##UBName##_##Name; \
+		SafeType##UBName##_##Name GetBindlessResource##UBName##_##Name() { return GetResourceFromHeap(SafeType##UBName##_##Name, UB_CB_PREFIXED_MEMBER_ACCESS(UBName, BindlessSRV_, Name)); } \
+		static const SafeType##UBName##_##Name UBName##_##Name = GetBindlessResource##UBName##_##Name();
 
 	#define UB_RESOURCE_MEMBER_UAV(Type, UBName, Name) \
-		Type GetBindlessResource##UBName##_##Name() { return GetResourceFromHeap(Type, UB_CB_PREFIXED_MEMBER_ACCESS(UBName, BindlessUAV_, Name)); } \
-		static const Type UBName##_##Name = GetBindlessResource##UBName##_##Name();
+		typedef Type SafeType##UBName##_##Name; \
+		SafeType##UBName##_##Name GetBindlessResource##UBName##_##Name() { return GetResourceFromHeap(SafeType##UBName##_##Name, UB_CB_PREFIXED_MEMBER_ACCESS(UBName, BindlessUAV_, Name)); } \
+		static const SafeType##UBName##_##Name UBName##_##Name = GetBindlessResource##UBName##_##Name();
 #else
 	#define UB_RESOURCE_MEMBER_SRV(Type, UBName, Name) Type UBName##_##Name
 	#define UB_RESOURCE_MEMBER_UAV(Type, UBName, Name) Type UBName##_##Name
