Fix for non matching shader structure at runtime due to missing shaders recompiled.

diff --git a/Engine/Shaders/Public/ShaderVersion.ush b/Engine/Shaders/Public/ShaderVersion.ush
--- a/Engine/Shaders/Public/ShaderVersion.ush
+++ b/Engine/Shaders/Public/ShaderVersion.ush
@@ -3,4 +3,4 @@
 // Each time the console command is executed it generates a new GUID. As this file is included
 // in Platform.ush (which should be included in any shader) it allows to invalidate the shader DDC.
 // 
-#pragma message("UESHADERMETADATA_VERSION 810370AA-2780-4601-9316-5186F4550D88")
+#pragma message("UESHADERMETADATA_VERSION 86603CDD-1211-4170-BF4A-EB98FD9F2D77")
