VulkanRHI: Add SupportsHighEndRayTracingEffects to DDPI for Vulkan ray tracing.  Add PLATFORM_SUPPORTS_CALLABLE_SHADERS to Vulkan shader compiler directly.

#rnx

diff --git a/Engine/Config/VulkanPC/DataDrivenPlatformInfo.ini b/Engine/Config/VulkanPC/DataDrivenPlatformInfo.ini
--- a/Engine/Config/VulkanPC/DataDrivenPlatformInfo.ini
+++ b/Engine/Config/VulkanPC/DataDrivenPlatformInfo.ini
@@ -182,6 +182,7 @@ bSupportsRayTracingShaders = true
 bSupportsPathTracing = true
 bSupportsRayTracingCallableShaders = true
 bSupportsRayTracingProceduralPrimitive = true
+bSupportsHighEndRayTracingEffects=true
 
 bSupportsGPUScene = true
 bSupportsPrimitiveShaders = false
diff --git a/Engine/Shaders/Public/Platform/Vulkan/VulkanCommon.ush b/Engine/Shaders/Public/Platform/Vulkan/VulkanCommon.ush
--- a/Engine/Shaders/Public/Platform/Vulkan/VulkanCommon.ush
+++ b/Engine/Shaders/Public/Platform/Vulkan/VulkanCommon.ush
@@ -7,7 +7,7 @@
 #pragma once
 
 // Update this GUID to invalidate shader recompilation for only Vulkan shaders
-#pragma message("UESHADERMETADATA_VERSION 78B8BEEB-B86D-4C6D-8541-BB93F687F063")
+#pragma message("UESHADERMETADATA_VERSION 52DDF0DD-4378-4376-9CFF-4DCC8271F475")
 
 #define COMPILER_SUPPORTS_ULONG_TYPES 1
 #define COMPILER_SUPPORTS_UINT64_IMAGE_ATOMICS (1 && COMPILER_SUPPORTS_ULONG_TYPES)
diff --git a/Engine/Source/Developer/VulkanShaderFormat/Private/VulkanShaderCompiler.cpp b/Engine/Source/Developer/VulkanShaderFormat/Private/VulkanShaderCompiler.cpp
--- a/Engine/Source/Developer/VulkanShaderFormat/Private/VulkanShaderCompiler.cpp
+++ b/Engine/Source/Developer/VulkanShaderFormat/Private/VulkanShaderCompiler.cpp
@@ -1411,6 +1411,7 @@ void ModifyVulkanCompilerInput(FShaderCompilerInput& Input)
 	else if (InternalState.IsSM6())
 	{
 		Input.Environment.SetDefine(TEXT("VULKAN_PROFILE_SM6"), 1);
+		Input.Environment.SetDefine(TEXT("PLATFORM_SUPPORTS_CALLABLE_SHADERS"), 1);
 	}
 	else if (InternalState.IsSM5())
 	{
