Bump shader version to flush out a suspected poisoned cache

diff --git a/Engine/Shaders/Public/ShaderVersion.ush b/Engine/Shaders/Public/ShaderVersion.ush
--- a/Engine/Shaders/Public/ShaderVersion.ush
+++ b/Engine/Shaders/Public/ShaderVersion.ush
@@ -3,4 +3,4 @@
 // Each time the console command is executed it generates a new GUID. As this file is included
 // in Platform.ush (which should be included in any shader) it allows to invalidate the shader DDC.
 // 
-#pragma message("UESHADERMETADATA_VERSION D6EB0A35-C442-4C75-89FD-C4D7D8E06556")
+#pragma message("UESHADERMETADATA_VERSION D5BA5133-72F2-4376-8C70-D1E927FFB436")
