Substrate - added missing Opacity output from the scene texture node as well as on the GBuffer overview debug view.

#rnx

diff --git a/Engine/Shaders/Public/SubstratePublic.ush b/Engine/Shaders/Public/SubstratePublic.ush
--- a/Engine/Shaders/Public/SubstratePublic.ush
+++ b/Engine/Shaders/Public/SubstratePublic.ush
@@ -62,7 +62,6 @@ FSubstrateGBuffer SubstratePublic_GetSubstrateGBufferFromFirstBSDF(uint2 PixelCo
 		SGBuffer.Specular			= SubstrateGetBSDFSpecular(BSDF);
 		SGBuffer.Metallic			= SubstrateGetBSDFMetallic(BSDF);
 		SGBuffer.WorldNormal		= SubstrateGetWorldNormal(SubstratePixelHeader, BSDF, SubstrateAddressing);
-		SGBuffer.Opacity			= BSDF.HasScatteringData() ? SubstrateSubSurfaceHeaderGetOpacity(SSSHeader) : 1.f;
 		SGBuffer.Roughness			= SubstrateGetBSDFRoughness(BSDF);
 		SGBuffer.MaterialAO			= SubstrateGetAO(SubstratePixelHeader);
 		SGBuffer.ShadingModelColor	= GetShadingModelColor(SubstrateGetLegacyShadingModels(BSDF));
@@ -71,6 +70,36 @@ FSubstrateGBuffer SubstratePublic_GetSubstrateGBufferFromFirstBSDF(uint2 PixelCo
 		SGBuffer.StoredSpecular		= SubstrateGetBSDFSpecular(BSDF).rrr;
 		SGBuffer.WorldTangent		= SubstrateGetWorldTangent(SubstratePixelHeader, BSDF, SubstrateAddressing);
 		SGBuffer.Anisotropy			= SubstrateGetBSDFAnisotropy(BSDF);
+
+		SGBuffer.Opacity			= 0.0f;
+		if (SubstrateSubSurfaceHeaderGetSSSType(SSSHeader) != SSS_TYPE_INVALID)
+		{
+			SGBuffer.Opacity = SubstrateSubSurfaceHeaderGetOpacity(SSSHeader);
+		}
+		else if (BSDF_GETTYPE(BSDF) == SUBSTRATE_BSDF_TYPE_SLAB)
+		{
+			if (BSDF_GETHASFUZZ(BSDF))
+			{
+				SGBuffer.Opacity = SLAB_FUZZ_AMOUNT(BSDF);	// Cloth
+			}
+			else if (BSDF_GETHASHAZINESS(BSDF))
+			{
+				const FHaziness Haziness = UnpackHaziness(SLAB_HAZINESS(BSDF));
+				const bool bHazeAsSimpleClearCoat = Haziness.bSimpleClearCoat;
+				if (bHazeAsSimpleClearCoat)
+				{
+					SGBuffer.Opacity = Haziness.Weight;		// ClearCoat
+				}
+			}
+		}
+		else if (BSDF_GETTYPE(BSDF) == SUBSTRATE_BSDF_TYPE_EYE)
+		{
+			SGBuffer.Opacity = 1 - EYE_IRISMASK(BSDF);		// IrisMask
+		}
+		else if (BSDF_GETTYPE(BSDF) == SUBSTRATE_BSDF_TYPE_HAIR)
+		{
+			SGBuffer.Opacity = HAIR_BACKLIT(BSDF);			// BackLit
+		}
 	}
 
 	return SGBuffer;
