[shaders] bump version to clear out bad cached shadermap data due to CL#35051183 (and subsequent backout)

#rnx

diff --git a/Engine/Shaders/Public/ShaderVersion.ush b/Engine/Shaders/Public/ShaderVersion.ush
--- a/Engine/Shaders/Public/ShaderVersion.ush
+++ b/Engine/Shaders/Public/ShaderVersion.ush
@@ -3,4 +3,4 @@
 // Each time the console command is executed it generates a new GUID. As this file is included
 // in Platform.ush (which should be included in any shader) it allows to invalidate the shader DDC.
 // 
-#pragma message("UESHADERMETADATA_VERSION 69D40495-E82A-4418-9D00-ED039CCD72DC")
+#pragma message("UESHADERMETADATA_VERSION 9DB6FF28-9CA1-4E83-97B3-9A355390EE10")
