Bump shader version to resolve deserialisation errors from previous backout.

#rnx

diff --git a/Engine/Shaders/Public/ShaderVersion.ush b/Engine/Shaders/Public/ShaderVersion.ush
--- a/Engine/Shaders/Public/ShaderVersion.ush
+++ b/Engine/Shaders/Public/ShaderVersion.ush
@@ -3,4 +3,4 @@
 // Each time the console command is executed it generates a new GUID. As this file is included
 // in Platform.ush (which should be included in any shader) it allows to invalidate the shader DDC.
 // 
-#pragma message("UESHADERMETADATA_VERSION 60ae06fd-b60f-4424-80dd-25b05e4faaa9")
+#pragma message("UESHADERMETADATA_VERSION A48469A4-4B75-2D09-9889-149BF44128E8")
\ No newline at end of file
