Adding OVERRIDE_ROOTCONSTANTS_USH for platforms to use their PlatformRootConstants.ush

diff --git a/Engine/Shaders/Public/RootConstants.ush b/Engine/Shaders/Public/RootConstants.ush
--- a/Engine/Shaders/Public/RootConstants.ush
+++ b/Engine/Shaders/Public/RootConstants.ush
@@ -3,10 +3,10 @@
 #pragma once
 
 #if PLATFORM_SUPPORTS_SHADER_ROOT_CONSTANTS
-	#if SM6_PROFILE || SM5_PROFILE
-		#include "Platform/D3D/D3DRootConstants.ush"
-	#else
+	#if OVERRIDE_ROOTCONSTANTS_USH
 		#include "/Platform/Public/PlatformRootConstants.ush"
+	#elif SM6_PROFILE || SM5_PROFILE
+		#include "Platform/D3D/D3DRootConstants.ush"
 	#endif
 #else
 	uint4 PassData; 
