Substrate - fixed access to gbuffer for some materials. (e.g. opaque materials cannot sample the gbuffer).
Also tested Niagara, deferred and forward for different projects.

#rnx

diff --git a/Engine/Plugins/FX/Niagara/Shaders/Private/NiagaraDataInterfaceCollisionQuery.ush b/Engine/Plugins/FX/Niagara/Shaders/Private/NiagaraDataInterfaceCollisionQuery.ush
--- a/Engine/Plugins/FX/Niagara/Shaders/Private/NiagaraDataInterfaceCollisionQuery.ush
+++ b/Engine/Plugins/FX/Niagara/Shaders/Private/NiagaraDataInterfaceCollisionQuery.ush
@@ -4,6 +4,8 @@
 
 #include "/Engine/Private/Common.ush"
 
+#define NEEDS_SCENE_TEXTURES 1
+#define SCENE_TEXTURES_DISABLED 0
 #define SubstratePublicStruct SubstratePublic
 #include "/Engine/Public/SubstratePublic.ush"
 
diff --git a/Engine/Plugins/FX/Niagara/Shaders/Private/NiagaraDataInterfaceGBufferTemplate.ush b/Engine/Plugins/FX/Niagara/Shaders/Private/NiagaraDataInterfaceGBufferTemplate.ush
--- a/Engine/Plugins/FX/Niagara/Shaders/Private/NiagaraDataInterfaceGBufferTemplate.ush
+++ b/Engine/Plugins/FX/Niagara/Shaders/Private/NiagaraDataInterfaceGBufferTemplate.ush
@@ -2,6 +2,8 @@
 
 #include "/Engine/Private/DeferredShadingCommon.ush"
 
+#define NEEDS_SCENE_TEXTURES 1
+#define SCENE_TEXTURES_DISABLED 0
 #define SubstratePublicStruct SubstratePublic
 #include "/Engine/Public/SubstratePublic.ush"
 
diff --git a/Engine/Plugins/FX/Niagara/Shaders/Private/NiagaraEmitterInstanceShader.usf b/Engine/Plugins/FX/Niagara/Shaders/Private/NiagaraEmitterInstanceShader.usf
--- a/Engine/Plugins/FX/Niagara/Shaders/Private/NiagaraEmitterInstanceShader.usf
+++ b/Engine/Plugins/FX/Niagara/Shaders/Private/NiagaraEmitterInstanceShader.usf
@@ -18,6 +18,9 @@ NiagaraSimulationShader.usf:
 #if GPU_SIMULATION
 	#include "/Engine/Private/Common.ush"
 	#include "/Engine/Private/DistanceField/GlobalDistanceFieldShared.ush"
+
+	#define NEEDS_SCENE_TEXTURES 1
+	#define SCENE_TEXTURES_DISABLED 0
 	#include "/Engine/Generated/UniformBuffers/SubstratePublic.ush"
 #else
 	const static float PI = 3.1415926535897932f;
diff --git a/Engine/Shaders/Public/SubstratePublic.ush b/Engine/Shaders/Public/SubstratePublic.ush
--- a/Engine/Shaders/Public/SubstratePublic.ush
+++ b/Engine/Shaders/Public/SubstratePublic.ush
@@ -8,7 +8,9 @@
 
 #include "/Engine/Private/Substrate/Substrate.ush"
 
-#if SUBSTRATE_ENABLED
+#if SUBSTRATE_ENABLED 
+
+#define SUBSTRATE_ALLOWS_GBUFFER_SAMPLING  ((NEEDS_SCENE_TEXTURES && !SCENE_TEXTURES_DISABLED) && SUBSTRATE_MATERIALCONTAINER_IS_VIEWRESOURCE)
 
 /** Public facing functions allowing external users access to some Substrate data. */
 
@@ -42,7 +44,9 @@ FSubstrateGBuffer SubstratePublic_GetSubstrateGBufferFromFirstBSDF(uint2 PixelCo
 {
 	FSubstrateGBuffer SGBuffer = (FSubstrateGBuffer)0; 
 	SGBuffer.bIsValid = false;
+	SGBuffer.WorldNormal = float3(0.0, 0.0, 1.0);
 
+#if SUBSTRATE_ALLOWS_GBUFFER_SAMPLING
 	FSubstrateAddressing SubstrateAddressing   = GetSubstratePixelDataByteOffset(PixelCoord, uint2(View.BufferSizeAndInvSize.xy), SubstratePublicStruct.MaxBytesPerPixel);
 	FSubstratePixelHeader SubstratePixelHeader = UnpackSubstrateHeaderIn(SubstratePublicStruct.MaterialTextureArray, SubstrateAddressing, SubstratePublicStruct.TopLayerTexture);
 	
@@ -101,6 +105,7 @@ FSubstrateGBuffer SubstratePublic_GetSubstrateGBufferFromFirstBSDF(uint2 PixelCo
 			SGBuffer.Opacity = HAIR_BACKLIT(BSDF);			// BackLit
 		}
 	}
+#endif
 
 	return SGBuffer;
 }
@@ -108,11 +113,19 @@ FSubstrateGBuffer SubstratePublic_GetSubstrateGBufferFromFirstBSDF(uint2 PixelCo
 // Optimised data fetch from the TopLayer texture
 float3 SubstratePublic_GetWorldNormal(uint2 PixelCoord)
 {
+#if SUBSTRATE_ALLOWS_GBUFFER_SAMPLING
   return SubstrateUnpackTopLayerData(SubstratePublicStruct.TopLayerTexture.Load(uint3(PixelCoord, 0))).WorldNormal;
+#else
+	return float3(0, 0, 1);
+#endif
 }
 float SubstratePublic_GetRoughness(uint2 PixelCoord)
 {
+#if SUBSTRATE_ALLOWS_GBUFFER_SAMPLING
   return SubstrateUnpackTopLayerData(SubstratePublicStruct.TopLayerTexture.Load(uint3(PixelCoord, 0))).Roughness;
+#else
+	return 0.0;
+#endif
 }
 
 #endif
\ No newline at end of file
