[Shaders] Bump ShaderVersion.ush after modifications in GlobalBeginCompileShader().

#rnx

diff --git a/Engine/Shaders/Public/ShaderVersion.ush b/Engine/Shaders/Public/ShaderVersion.ush
--- a/Engine/Shaders/Public/ShaderVersion.ush
+++ b/Engine/Shaders/Public/ShaderVersion.ush
@@ -3,4 +3,4 @@
 // Each time the console command is executed it generates a new GUID. As this file is included
 // in Platform.ush (which should be included in any shader) it allows to invalidate the shader DDC.
 // 
-#pragma message("UESHADERMETADATA_VERSION 1B9EB545-4184-F9A2-094F-A0AA4BD39EC4")
\ No newline at end of file
+#pragma message("UESHADERMETADATA_VERSION ECC0ACE5-64A8-47CA-A772-C34CD9E09C21")
\ No newline at end of file
