Support half (relaxed) precision in Android SM5 renderer

diff --git a/Engine/Shaders/Private/ClearCoatCommon.ush b/Engine/Shaders/Private/ClearCoatCommon.ush
--- a/Engine/Shaders/Private/ClearCoatCommon.ush
+++ b/Engine/Shaders/Private/ClearCoatCommon.ush
@@ -44,13 +44,13 @@ void RemapClearCoatDiffuseAndSpecularColor(FGBufferData GBuffer, half NoV, inout
 	}
 }
 
-void RemapClearCoatDiffuseAndSpecularColor(FGBufferData GBuffer, float2 ScreenPosition, inout float3 DiffuseColor, inout float3 SpecularColor)
+void RemapClearCoatDiffuseAndSpecularColor(FGBufferData GBuffer, float2 ScreenPosition, inout half3 DiffuseColor, inout half3 SpecularColor)
 {
 	if (GBuffer.ShadingModelID == SHADINGMODELID_CLEAR_COAT)
 	{
 		float3 TranslatedWorldPosition = mul(float4(GetScreenPositionForProjectionType(ScreenPosition, GBuffer.Depth), GBuffer.Depth, 1), View.ScreenToTranslatedWorld).xyz;
-		float3 V = -GetCameraVectorFromTranslatedWorldPosition(TranslatedWorldPosition);
-		float NoV = saturate(dot(GBuffer.WorldNormal, V));
+		half3 V = -GetCameraVectorFromTranslatedWorldPosition(TranslatedWorldPosition);
+		half NoV = saturate(dot(GBuffer.WorldNormal, V));
 
 		RemapClearCoatDiffuseAndSpecularColor(GBuffer, NoV, DiffuseColor, SpecularColor);
 	}
diff --git a/Engine/Shaders/Private/MaterialTemplate.ush b/Engine/Shaders/Private/MaterialTemplate.ush
--- a/Engine/Shaders/Private/MaterialTemplate.ush
+++ b/Engine/Shaders/Private/MaterialTemplate.ush
@@ -3846,14 +3846,14 @@ float3 GetMaterialPreviousWorldPositionOffset(FMaterialVertexParameters Paramete
 }
 
 // .rgb:SubsurfaceColor, .a:SSProfileId in 0..1 range
-half4 GetMaterialSubsurfaceDataRaw(FPixelMaterialInputs PixelMaterialInputs)
+MaterialFloat4 GetMaterialSubsurfaceDataRaw(FPixelMaterialInputs PixelMaterialInputs)
 {
 	return PixelMaterialInputs.Subsurface;
 }
 
-half4 GetMaterialSubsurfaceData(FPixelMaterialInputs PixelMaterialInputs)
+MaterialFloat4 GetMaterialSubsurfaceData(FPixelMaterialInputs PixelMaterialInputs)
 {
-	half4 OutSubsurface = GetMaterialSubsurfaceDataRaw(PixelMaterialInputs);
+	MaterialFloat4 OutSubsurface = GetMaterialSubsurfaceDataRaw(PixelMaterialInputs);
 	OutSubsurface.rgb = saturate(OutSubsurface.rgb);
 	return OutSubsurface;
 }
diff --git a/Engine/Shaders/Public/Platform.ush b/Engine/Shaders/Public/Platform.ush
--- a/Engine/Shaders/Public/Platform.ush
+++ b/Engine/Shaders/Public/Platform.ush
@@ -321,7 +321,7 @@
 #define FORCE_FLOATS 0
 #endif
 
-#if COMPILER_DXC && !FORCE_FLOATS && (ES3_1_PROFILE || METAL_PROFILE)
+#if COMPILER_DXC && !FORCE_FLOATS && (ES3_1_PROFILE || METAL_PROFILE || VULKAN_PROFILE_SM5)
 	#pragma dxc diagnostic ignored "-Wconversion"
 #endif
 
@@ -334,7 +334,7 @@
 	#define half3x3 min16float3x3
 	#define half3x4 min16float3x4
 	#define half4x4 min16float4x4
-#elif ((!(ES3_1_PROFILE || METAL_PROFILE)) && !PLATFORM_SUPPORTS_REAL_TYPES) || FORCE_FLOATS
+#elif ((!(ES3_1_PROFILE || METAL_PROFILE || VULKAN_PROFILE_SM5)) && !PLATFORM_SUPPORTS_REAL_TYPES) || FORCE_FLOATS
 	// Always use floats when using the ES3/METAL compiler, because platforms not optimized for lower precision, 
 	// And we don't want potential side effects on other platforms
 	#define half float
@@ -353,7 +353,7 @@
 	#define fixed3x3 float3x3
 	#define fixed4x4 float4x4
 	#define fixed4x3 float4x3
-#elif (VULKAN_PROFILE) || (COMPILER_GLSL_ES3_1 && !(COMPILER_HLSLCC && COMPILER_HLSLCC == 1))
+#elif (VULKAN_PROFILE || VULKAN_PROFILE_SM5) || (COMPILER_GLSL_ES3_1 && !(COMPILER_HLSLCC && COMPILER_HLSLCC == 1))
 	// For VULKAN and OPENGL ES31 use RelaxedPrecision for half floats
 	#define half min16float
 	#define half2 min16float2
@@ -697,7 +697,7 @@ void ShaderYield()
 #define COMPILER_SUPPORTS_MED3 0
 #endif
 
-#if ES3_1_PROFILE && !METAL_PROFILE
+#if (ES3_1_PROFILE && !METAL_PROFILE) || VULKAN_PROFILE_SM5
 #define HALF_TYPE half
 #define HALF2_TYPE half2
 #define HALF3_TYPE half3
diff --git a/Engine/Shaders/Public/Platform/Vulkan/VulkanCommon.ush b/Engine/Shaders/Public/Platform/Vulkan/VulkanCommon.ush
--- a/Engine/Shaders/Public/Platform/Vulkan/VulkanCommon.ush
+++ b/Engine/Shaders/Public/Platform/Vulkan/VulkanCommon.ush
@@ -7,7 +7,7 @@
 #pragma once
 
 // Update this GUID to invalidate shader recompilation for only Vulkan shaders
-#pragma message("UESHADERMETADATA_VERSION 368AA5D0-57AB-4A02-BE50-0E3432A4820D")
+#pragma message("UESHADERMETADATA_VERSION DE4FD0BE-F0E4-4B69-B942-F568445132B0")
 
 #define COMPILER_SUPPORTS_ULONG_TYPES 1
 #define COMPILER_SUPPORTS_UINT64_IMAGE_ATOMICS (1 && COMPILER_SUPPORTS_ULONG_TYPES)
