Reduce the loop unroll time for TSR shader compiler by adding manual unroll to frequently used paths.

diff --git a/Engine/Shaders/Public/LaneVectorization.ush b/Engine/Shaders/Public/LaneVectorization.ush
--- a/Engine/Shaders/Public/LaneVectorization.ush
+++ b/Engine/Shaders/Public/LaneVectorization.ush
@@ -604,8 +604,47 @@ void UnpackDwordsToVector(uint DW[2], out half4 V)
 
 #endif // PLATFORM_SUPPORTS_REAL_TYPES
 
+//------------------------------------------------------- UTIL FUNCTION FOR MANUAL LOOP UNROLL
 
+uint GetDW(uint DW[1], const uint Index)
+{
+	return DW[Index];
+}
+
+uint GetDW(uint DW[2], const uint Index)
+{
+	return DW[Index];
+}
+
+uint GetDW(uint DW[3], const uint Index)
+{
+	return DW[Index];
+}
 
+uint GetDW(uint DW[4], const uint Index)
+{
+	return DW[Index];
+}
+
+void SetDW(inout uint DW[1], const uint Index, uint V)
+{
+	DW[Index] = V;
+}
+
+void SetDW(inout uint DW[2], const uint Index, uint V)
+{
+	DW[Index] = V;
+}
+
+void SetDW(inout uint DW[3], const uint Index, uint V)
+{
+	DW[Index] = V;
+}
+
+void SetDW(inout uint DW[4], const uint Index, uint V)
+{
+	DW[Index] = V;
+}
 
 //------------------------------------------------------- GROUP SHARED COMMUNICATION
 
@@ -631,13 +670,27 @@ void AtomicIncrementLDSDword(const uint SharedIndex, uint V)
 template<typename ScalarType, uint VectorSize>
 void WriteVectorToLDS(const uint SharedIndex, vector<ScalarType, VectorSize> V)
 {
+	// manual loop unroll the most used cases to reduce compile time
+	// the branch will be compiled away
 	uint DW[GetVectorDwordSize(ScalarType, VectorSize)];
 	PackVectorToDwords(V, /* out */ DW);
-	
-	UNROLL
-	for (uint i = 0; i < GetVectorDwordSize(ScalarType, VectorSize); i++)
+
+	if (GetVectorDwordSize(ScalarType, VectorSize) == 1)
+	{
+		SharedData[SharedIndex + 0 * LDS_SIZE] = GetDW(DW, 0);
+	}
+	else if (GetVectorDwordSize(ScalarType, VectorSize) == 2)
 	{
-		SharedData[SharedIndex + i * LDS_SIZE] = DW[i];
+		SharedData[SharedIndex + 0 * LDS_SIZE] = GetDW(DW, 0);
+		SharedData[SharedIndex + 1 * LDS_SIZE] = GetDW(DW, 1);
+	}
+	else
+	{
+		UNROLL
+		for (uint i = 0; i < GetVectorDwordSize(ScalarType, VectorSize); i++)
+		{
+			SharedData[SharedIndex + i * LDS_SIZE] = GetDW(DW, i);
+		}
 	}
 }
 
@@ -645,10 +698,23 @@ template<typename ScalarType, uint VectorSize>
 void ReadVectorFromLDS(const uint SharedIndex, out vector<ScalarType, VectorSize> V)
 {
 	uint DW[GetVectorDwordSize(ScalarType, VectorSize)];
-	UNROLL
-	for (uint i = 0; i < GetVectorDwordSize(ScalarType, VectorSize); i++)
+
+	if (GetVectorDwordSize(ScalarType, VectorSize) == 1)
+	{
+		SetDW(DW, 0, SharedData[SharedIndex + 0 * LDS_SIZE]);
+	}
+	else if (GetVectorDwordSize(ScalarType, VectorSize) == 2)
 	{
-		DW[i] = SharedData[SharedIndex + i * LDS_SIZE];
+		SetDW(DW, 0, SharedData[SharedIndex + 0 * LDS_SIZE]);
+		SetDW(DW, 1, SharedData[SharedIndex + 1 * LDS_SIZE]);
+	}
+	else
+	{
+		UNROLL
+		for (uint i = 0; i < GetVectorDwordSize(ScalarType, VectorSize); i++)
+		{
+			SetDW(DW, i, SharedData[SharedIndex + i * LDS_SIZE]);
+		}
 	}
 	UnpackDwordsToVector(DW, /* out */ V);
 }
@@ -905,10 +971,23 @@ struct FRegisters_SoA_Matrix
 	vector<ScalarType, VectorSize> GetElement(const uint ElementIndex)
 	{
 		vector<ScalarType, VectorSize> Element;
-		UNROLL
-		for (uint ComponentIndex = 0; ComponentIndex < VectorSize; ComponentIndex++)
+		if (VectorSize == 1)
+		{
+			Element[0] = M[0][ElementIndex];
+		}
+		else if (VectorSize == 3)
+		{
+			Element[0] = M[0][ElementIndex];
+			Element[1] = M[1][ElementIndex];
+			Element[2] = M[2][ElementIndex];
+		}
+		else
 		{
-			Element[ComponentIndex] = M[ComponentIndex][ElementIndex];
+			UNROLL
+			for (uint ComponentIndex = 0; ComponentIndex < VectorSize; ComponentIndex++)
+			{
+				Element[ComponentIndex] = M[ComponentIndex][ElementIndex];
+			}
 		}
 		return Element;
 	}
@@ -916,20 +995,47 @@ struct FRegisters_SoA_Matrix
 	CALL_SITE_DEBUGLOC
 	void SetElement(const uint ElementIndex, vector<ScalarType, VectorSize> Element)
 	{
-		UNROLL
-		for (uint ComponentIndex = 0; ComponentIndex < VectorSize; ComponentIndex++)
+		if (VectorSize == 1)
 		{
-			M[ComponentIndex][ElementIndex] = Element[ComponentIndex];
+			M[0][ElementIndex] = Element[0];
+		}
+		else if (VectorSize == 3)
+		{
+			M[0][ElementIndex] = Element[0];
+			M[1][ElementIndex] = Element[1];
+			M[2][ElementIndex] = Element[2];
+		}
+		else
+		{
+			UNROLL
+			for (uint ComponentIndex = 0; ComponentIndex < VectorSize; ComponentIndex++)
+			{
+				M[ComponentIndex][ElementIndex] = Element[ComponentIndex];
+			}
 		}
 	}
 
 	vector<ScalarType, VectorSize> WaveBroadcastElement(const FWaveBroadcastSettings BroadcastSettings, const uint ElementIndex)
 	{
 		vector<ScalarType, VectorSize> Element;
-		UNROLL
-		for (uint ComponentIndex = 0; ComponentIndex < VectorSize; ComponentIndex++)
+		
+		if (VectorSize == 1)
+		{
+			Element[0] = WaveBroadcast(BroadcastSettings, M[0])[ElementIndex];
+		}
+		else if (VectorSize == 3)
+		{
+			Element[0] = WaveBroadcast(BroadcastSettings, M[0])[ElementIndex];
+			Element[1] = WaveBroadcast(BroadcastSettings, M[1])[ElementIndex];
+			Element[2] = WaveBroadcast(BroadcastSettings, M[2])[ElementIndex];
+		}
+		else
 		{
-			Element[ComponentIndex] = WaveBroadcast(BroadcastSettings, M[ComponentIndex])[ElementIndex];
+			UNROLL
+			for (uint ComponentIndex = 0; ComponentIndex < VectorSize; ComponentIndex++)
+			{
+				Element[ComponentIndex] = WaveBroadcast(BroadcastSettings, M[ComponentIndex])[ElementIndex];
+			}
 		}
 		return Element;
 	}
@@ -1183,10 +1289,25 @@ struct TLaneVector2D
 	CALL_SITE_DEBUGLOC
 	void SetAllElements(vector<ScalarType, VectorSize> Element)
 	{
-		UNROLL
-		for (uint ElementIndex = 0; ElementIndex < SimdSize; ElementIndex++)
+		if(SimdSize == 2)
 		{
-			Registers.SetElement(ElementIndex, Element);
+			Registers.SetElement(0, Element);
+			Registers.SetElement(1, Element);
+		}
+		else if (SimdSize == 4)
+		{
+			Registers.SetElement(0, Element);
+			Registers.SetElement(1, Element);
+			Registers.SetElement(2, Element);
+			Registers.SetElement(3, Element);
+		}
+		else
+		{
+			UNROLL
+			for (uint ElementIndex = 0; ElementIndex < SimdSize; ElementIndex++)
+			{
+				Registers.SetElement(ElementIndex, Element);
+			}
 		}
 	}
 	
@@ -1196,10 +1317,25 @@ struct TLaneVector2D
 	TLaneVector2D<ScalarType, 1, SimdSizeX, SimdSizeY> GetComponent(const uint ComponentIndex)
 	{
 		TLaneVector2D<ScalarType, 1, SimdSizeX, SimdSizeY> R;
-		UNROLL
-		for (uint ElementIndex = 0; ElementIndex < SimdSize; ElementIndex++)
+		if(SimdSize == 2)
 		{
-			R.SetElement(ElementIndex, Registers.GetElementComponent(ElementIndex, ComponentIndex));
+			R.SetElement(0, Registers.GetElementComponent(0, ComponentIndex));
+			R.SetElement(1, Registers.GetElementComponent(1, ComponentIndex));
+		}
+		else if (SimdSize == 4)
+		{
+			R.SetElement(0, Registers.GetElementComponent(0, ComponentIndex));
+			R.SetElement(1, Registers.GetElementComponent(1, ComponentIndex));
+			R.SetElement(2, Registers.GetElementComponent(2, ComponentIndex));
+			R.SetElement(3, Registers.GetElementComponent(3, ComponentIndex));
+		}
+		else
+		{
+			UNROLL
+			for (uint ElementIndex = 0; ElementIndex < SimdSize; ElementIndex++)
+			{
+				R.SetElement(ElementIndex, Registers.GetElementComponent(ElementIndex, ComponentIndex));
+			}
 		}
 		return R;
 	}
@@ -1207,10 +1343,25 @@ struct TLaneVector2D
 	CALL_SITE_DEBUGLOC
 	void SetComponent(uint ComponentIndex, TLaneVector2D<ScalarType, 1, SimdSizeX, SimdSizeY> Comp)
 	{
-		UNROLL
-		for (uint ElementIndex = 0; ElementIndex < SimdSize; ElementIndex++)
+		if(SimdSize == 2)
+		{
+			Registers.SetElementComponent(0, ComponentIndex, Comp.GetElement(0));
+			Registers.SetElementComponent(1, ComponentIndex, Comp.GetElement(1));
+		}
+		else if (SimdSize == 4)
 		{
-			Registers.SetElementComponent(ElementIndex, ComponentIndex, Comp.GetElement(ElementIndex));
+			Registers.SetElementComponent(0, ComponentIndex, Comp.GetElement(0));
+			Registers.SetElementComponent(1, ComponentIndex, Comp.GetElement(1));
+			Registers.SetElementComponent(2, ComponentIndex, Comp.GetElement(2));
+			Registers.SetElementComponent(3, ComponentIndex, Comp.GetElement(3));
+		}
+		else
+		{
+			UNROLL
+			for (uint ElementIndex = 0; ElementIndex < SimdSize; ElementIndex++)
+			{
+				Registers.SetElementComponent(ElementIndex, ComponentIndex, Comp.GetElement(ElementIndex));
+			}
 		}
 	}
 	
@@ -1229,10 +1380,25 @@ struct TLaneVector2D
 		TLaneVector2D<ScalarType, 1, SimdSizeX, SimdSizeY> A)
 	{
 		TLaneVector2D<ScalarType, VectorSize, SimdSizeX, SimdSizeY> R;
-		UNROLL
-		for (uint ElementIndex = 0; ElementIndex < SimdSize; ElementIndex++)
+		if(SimdSize == 2)
+		{
+			R.SetElement(0, A.GetElement(0));
+			R.SetElement(1, A.GetElement(1));
+		}
+		else if (SimdSize == 4)
+		{
+			R.SetElement(0, A.GetElement(0));
+			R.SetElement(1, A.GetElement(1));
+			R.SetElement(2, A.GetElement(2));
+			R.SetElement(3, A.GetElement(3));
+		}
+		else
 		{
-			R.SetElement(ElementIndex, A.GetElement(ElementIndex));
+			UNROLL
+			for (uint ElementIndex = 0; ElementIndex < SimdSize; ElementIndex++)
+			{
+				R.SetElement(ElementIndex, A.GetElement(ElementIndex));
+			}
 		}
 		return R;
 	}
@@ -1243,10 +1409,25 @@ struct TLaneVector2D
 		vector<ScalarType, VectorSize> A)
 	{
 		TLaneVector2D<ScalarType, VectorSize, SimdSizeX, SimdSizeY> R;
-		UNROLL
-		for (uint ElementIndex = 0; ElementIndex < SimdSize; ElementIndex++)
+		if(SimdSize == 2)
+		{
+			R.SetElement(0, A);
+			R.SetElement(1, A);
+		}
+		else if (SimdSize == 4)
 		{
-			R.SetElement(ElementIndex, A);
+			R.SetElement(0, A);
+			R.SetElement(1, A);
+			R.SetElement(2, A);
+			R.SetElement(3, A);
+		}
+		else
+		{
+			UNROLL
+			for (uint ElementIndex = 0; ElementIndex < SimdSize; ElementIndex++)
+			{
+				R.SetElement(ElementIndex, A);
+			}
 		}
 		return R;
 	}
@@ -1258,10 +1439,25 @@ struct TLaneVector2D
 		TLaneVector2D<SourceScalarType, 1, SimdSizeX, SimdSizeY> A)
 	{
 		TLaneVector2D<ScalarType, VectorSize, SimdSizeX, SimdSizeY> R;
-		UNROLL
-		for (uint ElementIndex = 0; ElementIndex < SimdSize; ElementIndex++)
+		if(SimdSize == 2)
 		{
-			R.SetElement(ElementIndex, vector<ScalarType, VectorSize>(A.GetElement(ElementIndex)));
+			R.SetElement(0, vector<ScalarType, VectorSize>(A.GetElement(0)));
+			R.SetElement(1, vector<ScalarType, VectorSize>(A.GetElement(1)));
+		}
+		else if (SimdSize == 4)
+		{
+			R.SetElement(0, vector<ScalarType, VectorSize>(A.GetElement(0)));
+			R.SetElement(1, vector<ScalarType, VectorSize>(A.GetElement(1)));
+			R.SetElement(2, vector<ScalarType, VectorSize>(A.GetElement(2)));
+			R.SetElement(3, vector<ScalarType, VectorSize>(A.GetElement(3)));
+		}
+		else
+		{
+			UNROLL
+			for (uint ElementIndex = 0; ElementIndex < SimdSize; ElementIndex++)
+			{
+				R.SetElement(ElementIndex, vector<ScalarType, VectorSize>(A.GetElement(ElementIndex)));
+			}
 		}
 		return R;
 	}
@@ -1271,10 +1467,25 @@ struct TLaneVector2D
 	void TightenRegisters()
 	#if TENSOR_REGISTER_TIGHTENING
 	{
-		UNROLL
-		for (uint RegisterRowIndex = 0; RegisterRowIndex < RegisterRowCount; RegisterRowIndex++)
+		if(RegisterRowCount == 2)
+		{
+			Registers.SetRegisterRow(0, PackRegistersTightly(Registers.GetRegisterRow(0)));
+			Registers.SetRegisterRow(1, PackRegistersTightly(Registers.GetRegisterRow(1)));
+		}
+		else if (RegisterRowCount == 4)
 		{
-			Registers.SetRegisterRow(RegisterRowIndex, PackRegistersTightly(Registers.GetRegisterRow(RegisterRowIndex)));
+			Registers.SetRegisterRow(0, PackRegistersTightly(Registers.GetRegisterRow(0)));
+			Registers.SetRegisterRow(1, PackRegistersTightly(Registers.GetRegisterRow(1)));
+			Registers.SetRegisterRow(2, PackRegistersTightly(Registers.GetRegisterRow(2)));
+			Registers.SetRegisterRow(3, PackRegistersTightly(Registers.GetRegisterRow(3)));
+		}
+		else
+		{
+			UNROLL
+			for (uint RegisterRowIndex = 0; RegisterRowIndex < RegisterRowCount; RegisterRowIndex++)
+			{
+				Registers.SetRegisterRow(RegisterRowIndex, PackRegistersTightly(Registers.GetRegisterRow(RegisterRowIndex)));
+			}
 		}
 	}
 	#else
@@ -1342,10 +1553,25 @@ struct TLaneVector2D
 	TLaneVector2D<ScalarType, VectorSize, SimdSizeX, SimdSizeY> operator + (TLaneVector2D<ScalarType, VectorSize, SimdSizeX, SimdSizeY> B)
 	{
 		TLaneVector2D<ScalarType, VectorSize, SimdSizeX, SimdSizeY> R;
-		UNROLL
-		for (uint RegisterRowIndex = 0; RegisterRowIndex < RegisterRowCount; RegisterRowIndex++)
+		if(RegisterRowCount == 2)
+		{
+			R.Registers.SetRegisterRow(0, Registers.GetRegisterRow(0) + B.Registers.GetRegisterRow(0));
+			R.Registers.SetRegisterRow(1, Registers.GetRegisterRow(1) + B.Registers.GetRegisterRow(1));
+		}
+		else if (RegisterRowCount == 4)
 		{
-			R.Registers.SetRegisterRow(RegisterRowIndex, Registers.GetRegisterRow(RegisterRowIndex) + B.Registers.GetRegisterRow(RegisterRowIndex));
+			R.Registers.SetRegisterRow(0, Registers.GetRegisterRow(0) + B.Registers.GetRegisterRow(0));
+			R.Registers.SetRegisterRow(1, Registers.GetRegisterRow(1) + B.Registers.GetRegisterRow(1));
+			R.Registers.SetRegisterRow(2, Registers.GetRegisterRow(2) + B.Registers.GetRegisterRow(2));
+			R.Registers.SetRegisterRow(3, Registers.GetRegisterRow(3) + B.Registers.GetRegisterRow(3));
+		}
+		else
+		{
+			UNROLL
+			for (uint RegisterRowIndex = 0; RegisterRowIndex < RegisterRowCount; RegisterRowIndex++)
+			{
+				R.Registers.SetRegisterRow(RegisterRowIndex, Registers.GetRegisterRow(RegisterRowIndex) + B.Registers.GetRegisterRow(RegisterRowIndex));
+			}
 		}
 		return R;
 	}
@@ -1357,10 +1583,25 @@ struct TLaneVector2D
 		BV.SetAllElements(B);
 
 		TLaneVector2D<ScalarType, VectorSize, SimdSizeX, SimdSizeY> R;
-		UNROLL
-		for (uint RegisterRowIndex = 0; RegisterRowIndex < RegisterRowCount; RegisterRowIndex++)
+		if(RegisterRowCount == 2)
+		{
+			R.Registers.SetRegisterRow(0, Registers.GetRegisterRow(0) + BV.Registers.GetRegisterRow(0));
+			R.Registers.SetRegisterRow(1, Registers.GetRegisterRow(1) + BV.Registers.GetRegisterRow(1));
+		}
+		else if (RegisterRowCount == 4)
 		{
-			R.Registers.SetRegisterRow(RegisterRowIndex, Registers.GetRegisterRow(RegisterRowIndex) + BV.Registers.GetRegisterRow(RegisterRowIndex));
+			R.Registers.SetRegisterRow(0, Registers.GetRegisterRow(0) + BV.Registers.GetRegisterRow(0));
+			R.Registers.SetRegisterRow(1, Registers.GetRegisterRow(1) + BV.Registers.GetRegisterRow(1));
+			R.Registers.SetRegisterRow(2, Registers.GetRegisterRow(2) + BV.Registers.GetRegisterRow(2));
+			R.Registers.SetRegisterRow(3, Registers.GetRegisterRow(3) + BV.Registers.GetRegisterRow(3));
+		}
+		else
+		{
+			UNROLL
+			for (uint RegisterRowIndex = 0; RegisterRowIndex < RegisterRowCount; RegisterRowIndex++)
+			{
+				R.Registers.SetRegisterRow(RegisterRowIndex, Registers.GetRegisterRow(RegisterRowIndex) + BV.Registers.GetRegisterRow(RegisterRowIndex));
+			}
 		}
 		return R;
 	}
@@ -1371,10 +1612,25 @@ struct TLaneVector2D
 	TLaneVector2D<ScalarType, VectorSize, SimdSizeX, SimdSizeY> operator - ()
 	{
 		TLaneVector2D<ScalarType, VectorSize, SimdSizeX, SimdSizeY> R;
-		UNROLL
-		for (uint RegisterRowIndex = 0; RegisterRowIndex < RegisterRowCount; RegisterRowIndex++)
+		if(RegisterRowCount == 2)
 		{
-			R.Registers.SetRegisterRow(RegisterRowIndex, -Registers.GetRegisterRow(RegisterRowIndex));
+			R.Registers.SetRegisterRow(0, -Registers.GetRegisterRow(0));
+			R.Registers.SetRegisterRow(1, -Registers.GetRegisterRow(1));
+		}
+		else if (RegisterRowCount == 4)
+		{
+			R.Registers.SetRegisterRow(0, -Registers.GetRegisterRow(0));
+			R.Registers.SetRegisterRow(1, -Registers.GetRegisterRow(1));
+			R.Registers.SetRegisterRow(2, -Registers.GetRegisterRow(2));
+			R.Registers.SetRegisterRow(3, -Registers.GetRegisterRow(3));
+		}
+		else
+		{
+			UNROLL
+			for (uint RegisterRowIndex = 0; RegisterRowIndex < RegisterRowCount; RegisterRowIndex++)
+			{
+				R.Registers.SetRegisterRow(RegisterRowIndex, -Registers.GetRegisterRow(RegisterRowIndex));
+			}
 		}
 		return R;
 	}
@@ -1384,10 +1640,25 @@ struct TLaneVector2D
 	TLaneVector2D<ScalarType, VectorSize, SimdSizeX, SimdSizeY> operator - (TLaneVector2D<ScalarType, VectorSize, SimdSizeX, SimdSizeY> B)
 	{
 		TLaneVector2D<ScalarType, VectorSize, SimdSizeX, SimdSizeY> R;
-		UNROLL
-		for (uint RegisterRowIndex = 0; RegisterRowIndex < RegisterRowCount; RegisterRowIndex++)
+		if(RegisterRowCount == 2)
+		{
+			R.Registers.SetRegisterRow(0, Registers.GetRegisterRow(0) - B.Registers.GetRegisterRow(0));
+			R.Registers.SetRegisterRow(1, Registers.GetRegisterRow(1) - B.Registers.GetRegisterRow(1));
+		}
+		else if (RegisterRowCount == 4)
 		{
-			R.Registers.SetRegisterRow(RegisterRowIndex, Registers.GetRegisterRow(RegisterRowIndex) - B.Registers.GetRegisterRow(RegisterRowIndex));
+			R.Registers.SetRegisterRow(0, Registers.GetRegisterRow(0) - B.Registers.GetRegisterRow(0));
+			R.Registers.SetRegisterRow(1, Registers.GetRegisterRow(1) - B.Registers.GetRegisterRow(1));
+			R.Registers.SetRegisterRow(2, Registers.GetRegisterRow(2) - B.Registers.GetRegisterRow(2));
+			R.Registers.SetRegisterRow(3, Registers.GetRegisterRow(3) - B.Registers.GetRegisterRow(3));
+		}
+		else
+		{
+			UNROLL
+			for (uint RegisterRowIndex = 0; RegisterRowIndex < RegisterRowCount; RegisterRowIndex++)
+			{
+				R.Registers.SetRegisterRow(RegisterRowIndex, Registers.GetRegisterRow(RegisterRowIndex) - B.Registers.GetRegisterRow(RegisterRowIndex));
+			}
 		}
 		return R;
 	}
@@ -1399,10 +1670,25 @@ struct TLaneVector2D
 		BV.SetAllElements(B);
 
 		TLaneVector2D<ScalarType, VectorSize, SimdSizeX, SimdSizeY> R;
-		UNROLL
-		for (uint RegisterRowIndex = 0; RegisterRowIndex < RegisterRowCount; RegisterRowIndex++)
+		if(RegisterRowCount == 2)
+		{
+			R.Registers.SetRegisterRow(0, Registers.GetRegisterRow(0) - BV.Registers.GetRegisterRow(0));
+			R.Registers.SetRegisterRow(1, Registers.GetRegisterRow(1) - BV.Registers.GetRegisterRow(1));
+		}
+		else if (RegisterRowCount == 4)
 		{
-			R.Registers.SetRegisterRow(RegisterRowIndex, Registers.GetRegisterRow(RegisterRowIndex) - BV.Registers.GetRegisterRow(RegisterRowIndex));
+			R.Registers.SetRegisterRow(0, Registers.GetRegisterRow(0) - BV.Registers.GetRegisterRow(0));
+			R.Registers.SetRegisterRow(1, Registers.GetRegisterRow(1) - BV.Registers.GetRegisterRow(1));
+			R.Registers.SetRegisterRow(2, Registers.GetRegisterRow(2) - BV.Registers.GetRegisterRow(2));
+			R.Registers.SetRegisterRow(3, Registers.GetRegisterRow(3) - BV.Registers.GetRegisterRow(3));
+		}
+		else
+		{
+			UNROLL
+			for (uint RegisterRowIndex = 0; RegisterRowIndex < RegisterRowCount; RegisterRowIndex++)
+			{
+				R.Registers.SetRegisterRow(RegisterRowIndex, Registers.GetRegisterRow(RegisterRowIndex) - BV.Registers.GetRegisterRow(RegisterRowIndex));
+			}
 		}
 		return R;
 	}
@@ -1413,10 +1699,25 @@ struct TLaneVector2D
 	TLaneVector2D<ScalarType, VectorSize, SimdSizeX, SimdSizeY> operator * (TLaneVector2D<ScalarType, VectorSize, SimdSizeX, SimdSizeY> B)
 	{
 		TLaneVector2D<ScalarType, VectorSize, SimdSizeX, SimdSizeY> R;
-		UNROLL
-		for (uint RegisterRowIndex = 0; RegisterRowIndex < RegisterRowCount; RegisterRowIndex++)
+		if(RegisterRowCount == 2)
 		{
-			R.Registers.SetRegisterRow(RegisterRowIndex, Registers.GetRegisterRow(RegisterRowIndex) * B.Registers.GetRegisterRow(RegisterRowIndex));
+			R.Registers.SetRegisterRow(0, Registers.GetRegisterRow(0) * B.Registers.GetRegisterRow(0));
+			R.Registers.SetRegisterRow(1, Registers.GetRegisterRow(1) * B.Registers.GetRegisterRow(1));
+		}
+		else if (RegisterRowCount == 4)
+		{
+			R.Registers.SetRegisterRow(0, Registers.GetRegisterRow(0) * B.Registers.GetRegisterRow(0));
+			R.Registers.SetRegisterRow(1, Registers.GetRegisterRow(1) * B.Registers.GetRegisterRow(1));
+			R.Registers.SetRegisterRow(2, Registers.GetRegisterRow(2) * B.Registers.GetRegisterRow(2));
+			R.Registers.SetRegisterRow(3, Registers.GetRegisterRow(3) * B.Registers.GetRegisterRow(3));
+		}
+		else
+		{
+			UNROLL
+			for (uint RegisterRowIndex = 0; RegisterRowIndex < RegisterRowCount; RegisterRowIndex++)
+			{
+				R.Registers.SetRegisterRow(RegisterRowIndex, Registers.GetRegisterRow(RegisterRowIndex) * B.Registers.GetRegisterRow(RegisterRowIndex));
+			}
 		}
 		return R;
 	}
@@ -1428,10 +1729,25 @@ struct TLaneVector2D
 		BV.SetAllElements(B);
 
 		TLaneVector2D<ScalarType, VectorSize, SimdSizeX, SimdSizeY> R;
-		UNROLL
-		for (uint RegisterRowIndex = 0; RegisterRowIndex < RegisterRowCount; RegisterRowIndex++)
+		if(RegisterRowCount == 2)
+		{
+			R.Registers.SetRegisterRow(0, Registers.GetRegisterRow(0) * BV.Registers.GetRegisterRow(0));
+			R.Registers.SetRegisterRow(1, Registers.GetRegisterRow(1) * BV.Registers.GetRegisterRow(1));
+		}
+		else if (RegisterRowCount == 4)
 		{
-			R.Registers.SetRegisterRow(RegisterRowIndex, Registers.GetRegisterRow(RegisterRowIndex) * BV.Registers.GetRegisterRow(RegisterRowIndex));
+			R.Registers.SetRegisterRow(0, Registers.GetRegisterRow(0) * BV.Registers.GetRegisterRow(0));
+			R.Registers.SetRegisterRow(1, Registers.GetRegisterRow(1) * BV.Registers.GetRegisterRow(1));
+			R.Registers.SetRegisterRow(2, Registers.GetRegisterRow(2) * BV.Registers.GetRegisterRow(2));
+			R.Registers.SetRegisterRow(3, Registers.GetRegisterRow(3) * BV.Registers.GetRegisterRow(3));
+		}
+		else
+		{
+			UNROLL
+			for (uint RegisterRowIndex = 0; RegisterRowIndex < RegisterRowCount; RegisterRowIndex++)
+			{
+				R.Registers.SetRegisterRow(RegisterRowIndex, Registers.GetRegisterRow(RegisterRowIndex) * BV.Registers.GetRegisterRow(RegisterRowIndex));
+			}
 		}
 		return R;
 	}
@@ -1441,10 +1757,25 @@ struct TLaneVector2D
 	TLaneVector2D<bool, VectorSize, SimdSizeX, SimdSizeY> operator > (TLaneVector2D<ScalarType, VectorSize, SimdSizeX, SimdSizeY> B)
 	{
 		TLaneVector2D<bool, VectorSize, SimdSizeX, SimdSizeY> R;
-		UNROLL
-		for (uint RegisterRowIndex = 0; RegisterRowIndex < RegisterRowCount; RegisterRowIndex++)
+		if(RegisterRowCount == 2)
+		{
+			R.Registers.SetRegisterRow(0, Registers.GetRegisterRow(0) > B.Registers.GetRegisterRow(0));
+			R.Registers.SetRegisterRow(1, Registers.GetRegisterRow(1) > B.Registers.GetRegisterRow(1));
+		}
+		else if(RegisterRowCount == 4)
 		{
-			R.Registers.SetRegisterRow(RegisterRowIndex, Registers.GetRegisterRow(RegisterRowIndex) > B.Registers.GetRegisterRow(RegisterRowIndex));
+			R.Registers.SetRegisterRow(0, Registers.GetRegisterRow(0) > B.Registers.GetRegisterRow(0));
+			R.Registers.SetRegisterRow(1, Registers.GetRegisterRow(1) > B.Registers.GetRegisterRow(1));
+			R.Registers.SetRegisterRow(2, Registers.GetRegisterRow(2) > B.Registers.GetRegisterRow(2));
+			R.Registers.SetRegisterRow(3, Registers.GetRegisterRow(3) > B.Registers.GetRegisterRow(3));
+		}
+		else
+		{
+			UNROLL
+			for (uint RegisterRowIndex = 0; RegisterRowIndex < RegisterRowCount; RegisterRowIndex++)
+			{
+				R.Registers.SetRegisterRow(RegisterRowIndex, Registers.GetRegisterRow(RegisterRowIndex) > B.Registers.GetRegisterRow(RegisterRowIndex));
+			}
 		}
 		return R;
 	}
@@ -1453,10 +1784,26 @@ struct TLaneVector2D
 	TLaneVector2D<bool, VectorSize, SimdSizeX, SimdSizeY> operator < (TLaneVector2D<ScalarType, VectorSize, SimdSizeX, SimdSizeY> B)
 	{
 		TLaneVector2D<bool, VectorSize, SimdSizeX, SimdSizeY> R;
-		UNROLL
-		for (uint RegisterRowIndex = 0; RegisterRowIndex < RegisterRowCount; RegisterRowIndex++)
+
+		if(RegisterRowCount == 2)
+		{
+			R.Registers.SetRegisterRow(0, Registers.GetRegisterRow(0) < B.Registers.GetRegisterRow(0));
+			R.Registers.SetRegisterRow(1, Registers.GetRegisterRow(1) < B.Registers.GetRegisterRow(1));
+		}
+		else if(RegisterRowCount == 4)
 		{
-			R.Registers.SetRegisterRow(RegisterRowIndex, Registers.GetRegisterRow(RegisterRowIndex) < B.Registers.GetRegisterRow(RegisterRowIndex));
+			R.Registers.SetRegisterRow(0, Registers.GetRegisterRow(0) < B.Registers.GetRegisterRow(0));
+			R.Registers.SetRegisterRow(1, Registers.GetRegisterRow(1) < B.Registers.GetRegisterRow(1));
+			R.Registers.SetRegisterRow(2, Registers.GetRegisterRow(2) < B.Registers.GetRegisterRow(2));
+			R.Registers.SetRegisterRow(3, Registers.GetRegisterRow(3) < B.Registers.GetRegisterRow(3));
+		}
+		else
+		{
+			UNROLL
+			for (uint RegisterRowIndex = 0; RegisterRowIndex < RegisterRowCount; RegisterRowIndex++)
+			{
+				R.Registers.SetRegisterRow(RegisterRowIndex, Registers.GetRegisterRow(RegisterRowIndex) < B.Registers.GetRegisterRow(RegisterRowIndex));
+			}
 		}
 		return R;
 	}
@@ -1606,8 +1953,23 @@ void Deconcatenate(
 		TLaneVector2D<ScalarType, VectorSize, SimdSizeX, SimdSizeY> A) \
 	{ \
 		TLaneVector2D<ReturnScalarType, VectorSize, SimdSizeX, SimdSizeY> R; \
-		UNROLL for (uint RegisterRowIndex = 0; RegisterRowIndex < TLaneVector2D<ScalarType, VectorSize, SimdSizeX, SimdSizeY>::RegisterRowCount; RegisterRowIndex++) { \
-			R.Registers.SetRegisterRow(RegisterRowIndex, FunctionName(A.Registers.GetRegisterRow(RegisterRowIndex))); \
+		if(TLaneVector2D<ScalarType, VectorSize, SimdSizeX, SimdSizeY>::RegisterRowCount == 2) \
+		{ \
+			R.Registers.SetRegisterRow(0, FunctionName(A.Registers.GetRegisterRow(0))); \
+			R.Registers.SetRegisterRow(1, FunctionName(A.Registers.GetRegisterRow(1))); \
+		} \
+		else if(TLaneVector2D<ScalarType, VectorSize, SimdSizeX, SimdSizeY>::RegisterRowCount == 4) \
+		{ \
+			R.Registers.SetRegisterRow(0, FunctionName(A.Registers.GetRegisterRow(0))); \
+			R.Registers.SetRegisterRow(1, FunctionName(A.Registers.GetRegisterRow(1))); \
+			R.Registers.SetRegisterRow(2, FunctionName(A.Registers.GetRegisterRow(2))); \
+			R.Registers.SetRegisterRow(3, FunctionName(A.Registers.GetRegisterRow(3))); \
+		} \
+		else \
+		{ \
+			UNROLL for (uint RegisterRowIndex = 0; RegisterRowIndex < TLaneVector2D<ScalarType, VectorSize, SimdSizeX, SimdSizeY>::RegisterRowCount; RegisterRowIndex++) { \
+				R.Registers.SetRegisterRow(RegisterRowIndex, FunctionName(A.Registers.GetRegisterRow(RegisterRowIndex))); \
+			} \
 		} \
 		return R; \
 	} \
@@ -1641,8 +2003,23 @@ TVECTOR_FUNCTION_1PARAMS(uint, asuint);
 		TLaneVector2D<ScalarType, VectorSize, SimdSizeX, SimdSizeY> B) \
 	{ \
 		TLaneVector2D<ScalarType, VectorSize, SimdSizeX, SimdSizeY> R; \
-		UNROLL for (uint RegisterRowIndex = 0; RegisterRowIndex < TLaneVector2D<ScalarType, VectorSize, SimdSizeX, SimdSizeY>::RegisterRowCount; RegisterRowIndex++) { \
-			R.Registers.SetRegisterRow(RegisterRowIndex, FunctionName(A.Registers.GetRegisterRow(RegisterRowIndex), B.Registers.GetRegisterRow(RegisterRowIndex))); \
+		if (TLaneVector2D<ScalarType, VectorSize, SimdSizeX, SimdSizeY>::RegisterRowCount == 2) \
+		{ \
+			R.Registers.SetRegisterRow(0, FunctionName(A.Registers.GetRegisterRow(0), B.Registers.GetRegisterRow(0))); \
+			R.Registers.SetRegisterRow(1, FunctionName(A.Registers.GetRegisterRow(1), B.Registers.GetRegisterRow(1))); \
+		} \
+		else if (TLaneVector2D<ScalarType, VectorSize, SimdSizeX, SimdSizeY>::RegisterRowCount == 4) \
+		{ \
+			R.Registers.SetRegisterRow(0, FunctionName(A.Registers.GetRegisterRow(0), B.Registers.GetRegisterRow(0))); \
+			R.Registers.SetRegisterRow(1, FunctionName(A.Registers.GetRegisterRow(1), B.Registers.GetRegisterRow(1))); \
+			R.Registers.SetRegisterRow(2, FunctionName(A.Registers.GetRegisterRow(2), B.Registers.GetRegisterRow(2))); \
+			R.Registers.SetRegisterRow(3, FunctionName(A.Registers.GetRegisterRow(3), B.Registers.GetRegisterRow(3))); \
+		} \
+		else \
+		{ \
+			UNROLL for (uint RegisterRowIndex = 0; RegisterRowIndex < TLaneVector2D<ScalarType, VectorSize, SimdSizeX, SimdSizeY>::RegisterRowCount; RegisterRowIndex++) { \
+				R.Registers.SetRegisterRow(RegisterRowIndex, FunctionName(A.Registers.GetRegisterRow(RegisterRowIndex), B.Registers.GetRegisterRow(RegisterRowIndex))); \
+			} \
 		} \
 		return R; \
 	} \
@@ -1668,8 +2045,23 @@ TVECTOR_FUNCTION_2PARAMS(bit_shift_right);
 		TLaneVector2D<ScalarType, VectorSize, SimdSizeX, SimdSizeY> B) \
 	{ \
 		TLaneVector2D<ScalarType, ReturnVectorSize, SimdSizeX, SimdSizeY> R; \
-		UNROLL for (uint ElementIndex = 0; ElementIndex < SimdSizeX * SimdSizeY; ElementIndex++) { \
-			R.SetElement(ElementIndex, FunctionName(A.GetElement(ElementIndex), B.GetElement(ElementIndex))); \
+		if (SimdSizeX * SimdSizeY == 2) \
+		{ \
+			R.SetElement(0, FunctionName(A.GetElement(0), B.GetElement(0))); \
+			R.SetElement(1, FunctionName(A.GetElement(1), B.GetElement(1))); \
+		} \
+		else if (SimdSizeX * SimdSizeY == 4) \
+		{ \
+			R.SetElement(0, FunctionName(A.GetElement(0), B.GetElement(0))); \
+			R.SetElement(1, FunctionName(A.GetElement(1), B.GetElement(1))); \
+			R.SetElement(2, FunctionName(A.GetElement(2), B.GetElement(2))); \
+			R.SetElement(3, FunctionName(A.GetElement(3), B.GetElement(3))); \
+		} \
+		else \
+		{ \
+			UNROLL for (uint ElementIndex = 0; ElementIndex < SimdSizeX * SimdSizeY; ElementIndex++) { \
+				R.SetElement(ElementIndex, FunctionName(A.GetElement(ElementIndex), B.GetElement(ElementIndex))); \
+			} \
 		} \
 		return R; \
 	} \
@@ -1687,8 +2079,23 @@ TVECTOR_FUNCTION_2PARAMS_DIFF_RETURN(/* VectorSize = */ 1, dot);
 		TLaneVector2D<ScalarType, VectorSize, SimdSizeX, SimdSizeY> C) \
 	{ \
 		TLaneVector2D<ScalarType, VectorSize, SimdSizeX, SimdSizeY> R; \
-		UNROLL for (uint RegisterRowIndex = 0; RegisterRowIndex < TLaneVector2D<ScalarType, VectorSize, SimdSizeX, SimdSizeY>::RegisterRowCount; RegisterRowIndex++) { \
-			R.Registers.SetRegisterRow(RegisterRowIndex, FunctionName(A.Registers.GetRegisterRow(RegisterRowIndex), B.Registers.GetRegisterRow(RegisterRowIndex), C.Registers.GetRegisterRow(RegisterRowIndex))); \
+		if (TLaneVector2D<ScalarType, VectorSize, SimdSizeX, SimdSizeY>::RegisterRowCount == 2) \
+		{ \
+			R.Registers.SetRegisterRow(0, FunctionName(A.Registers.GetRegisterRow(0), B.Registers.GetRegisterRow(0), C.Registers.GetRegisterRow(0))); \
+			R.Registers.SetRegisterRow(1, FunctionName(A.Registers.GetRegisterRow(1), B.Registers.GetRegisterRow(1), C.Registers.GetRegisterRow(1))); \
+		} \
+		else if (TLaneVector2D<ScalarType, VectorSize, SimdSizeX, SimdSizeY>::RegisterRowCount == 4) \
+		{ \
+			R.Registers.SetRegisterRow(0, FunctionName(A.Registers.GetRegisterRow(0), B.Registers.GetRegisterRow(0), C.Registers.GetRegisterRow(0))); \
+			R.Registers.SetRegisterRow(1, FunctionName(A.Registers.GetRegisterRow(1), B.Registers.GetRegisterRow(1), C.Registers.GetRegisterRow(1))); \
+			R.Registers.SetRegisterRow(2, FunctionName(A.Registers.GetRegisterRow(2), B.Registers.GetRegisterRow(2), C.Registers.GetRegisterRow(2))); \
+			R.Registers.SetRegisterRow(3, FunctionName(A.Registers.GetRegisterRow(3), B.Registers.GetRegisterRow(3), C.Registers.GetRegisterRow(3))); \
+		} \
+		else \
+		{ \
+			UNROLL for (uint RegisterRowIndex = 0; RegisterRowIndex < TLaneVector2D<ScalarType, VectorSize, SimdSizeX, SimdSizeY>::RegisterRowCount; RegisterRowIndex++) { \
+				R.Registers.SetRegisterRow(RegisterRowIndex, FunctionName(A.Registers.GetRegisterRow(RegisterRowIndex), B.Registers.GetRegisterRow(RegisterRowIndex), C.Registers.GetRegisterRow(RegisterRowIndex))); \
+			} \
 		} \
 		return R; \
 	} \
@@ -1710,6 +2117,20 @@ template<uint VectorSize, uint SimdSizeX, uint SimdSizeY>
 bool AnyComponent(TLaneVector2D<bool, VectorSize, SimdSizeX, SimdSizeY> A)
 {
 	bool R = any(A.GetElement(0));
-	UNROLL for (uint SimdIndex = 1; SimdIndex < SimdSizeX * SimdSizeY; SimdIndex++) { R = R || any(A.GetElement(SimdIndex)); }
+	
+	if (SimdSizeX * SimdSizeY == 2)
+	{
+		R = R || any(A.GetElement(1));
+	}	
+	else if (SimdSizeX * SimdSizeY == 4)
+	{
+		R = R || any(A.GetElement(1));
+		R = R || any(A.GetElement(2));
+		R = R || any(A.GetElement(3));
+	}	
+	else
+	{
+		UNROLL for (uint SimdIndex = 1; SimdIndex < SimdSizeX * SimdSizeY; SimdIndex++) { R = R || any(A.GetElement(SimdIndex)); }
+	}
 	return R;
 }
