Bindless Fixes
- Allowing platforms to include a PlatformBindlessResources.ush file after all the BindlessResources.ush code.
- Moving RT bindless code from D3DCommon to PlatformBindlessResources.ush to make sure it uses the defines *after* they are defined.

diff --git a/Engine/Shaders/Public/BindlessResources.ush b/Engine/Shaders/Public/BindlessResources.ush
--- a/Engine/Shaders/Public/BindlessResources.ush
+++ b/Engine/Shaders/Public/BindlessResources.ush
@@ -85,3 +85,9 @@ TEXTURE_FROM_COLLECTION_T(TextureCubeArray)
 TEXTURE_FROM_COLLECTION_T(Texture3D)
 
 #endif
+
+#if defined(INCLUDE_PLATFORM_BINDLESS_RESOURCES_USH)
+	#include "/Platform/Public/PlatformBindlessResources.ush"
+#elif SM6_PROFILE
+	#include "Platform/D3D/PlatformBindlessResources.ush"
+#endif
diff --git a/Engine/Shaders/Public/Platform/D3D/D3DCommon.ush b/Engine/Shaders/Public/Platform/D3D/D3DCommon.ush
--- a/Engine/Shaders/Public/Platform/D3D/D3DCommon.ush
+++ b/Engine/Shaders/Public/Platform/D3D/D3DCommon.ush
@@ -148,27 +148,3 @@ void UEReportAssertWithPayload(uint AssertID, uint4 Payload)
 	#define GetResourceFromHeap(Type, Index) ResourceDescriptorHeap[Index]
 	#define GetSamplerFromHeap(Type, Index) SamplerDescriptorHeap[Index]
 #endif
-	
-
-#if ENABLE_BINDLESS_RESOURCES && PLATFORM_SUPPORTS_BINDLESS && (RAYCALLABLESHADER || RAYHITGROUPSHADER || RAYGENSHADER || RAYMISSSHADER)
-	#define OVERRIDE_RAY_TRACING_HIT_GROUP_SYSTEM_RESOURCES 1
-	#include "/Engine/Shared/RayTracingBuiltInResources.h"
-
-	// Combine the common FHitGroupSystemRootConstants with bindless parameters
-	struct FD3DHitGroupSystemParameters
-	{
-		FHitGroupSystemRootConstants RootConstants;
-
-		uint BindlessHitGroupSystemIndexBuffer;
-		uint BindlessHitGroupSystemVertexBuffer;
-	};
-
-	ConstantBuffer<FD3DHitGroupSystemParameters>	D3DHitGroupSystemParameters : UE_HLSL_REGISTER(b, RAY_TRACING_SYSTEM_ROOTCONSTANT_REGISTER, UE_HLSL_SPACE_RAY_TRACING_SYSTEM);
-	
-	static const FHitGroupSystemRootConstants HitGroupSystemRootConstants = D3DHitGroupSystemParameters.RootConstants;
-
-	// Will get filled in with the bindless index in D3DHitGroupSystemParameters.BindlessHitGroupSystemIndexBuffer
-	ByteAddressBuffer HitGroupSystemIndexBuffer;
-	// Will get filled in with the bindless index in D3DHitGroupSystemParameters.BindlessHitGroupSystemVertexBuffer
-	ByteAddressBuffer HitGroupSystemVertexBuffer;
-#endif
diff --git a/Engine/Shaders/Public/Platform/D3D/PlatformBindlessResources.ush b/Engine/Shaders/Public/Platform/D3D/PlatformBindlessResources.ush
new file mode 100644
--- /dev/null
+++ b/Engine/Shaders/Public/Platform/D3D/PlatformBindlessResources.ush
@@ -0,0 +1,27 @@
+// Copyright Epic Games, Inc. All Rights Reserved.
+
+#pragma once
+
+#if ENABLE_BINDLESS_RESOURCES && PLATFORM_SUPPORTS_BINDLESS && (RAYCALLABLESHADER || RAYHITGROUPSHADER || RAYGENSHADER || RAYMISSSHADER)
+	#define OVERRIDE_RAY_TRACING_HIT_GROUP_SYSTEM_RESOURCES 1
+	#include "/Engine/Shared/RayTracingBuiltInResources.h"
+
+	// Combine the common FHitGroupSystemRootConstants with bindless parameters
+	struct FD3DHitGroupSystemParameters
+	{
+		FHitGroupSystemRootConstants RootConstants;
+
+		uint BindlessHitGroupSystemIndexBuffer;
+		uint BindlessHitGroupSystemVertexBuffer;
+	};
+
+	ConstantBuffer<FD3DHitGroupSystemParameters>	D3DHitGroupSystemParameters : UE_HLSL_REGISTER(b, RAY_TRACING_SYSTEM_ROOTCONSTANT_REGISTER, UE_HLSL_SPACE_RAY_TRACING_SYSTEM);
+	
+	static const FHitGroupSystemRootConstants HitGroupSystemRootConstants = D3DHitGroupSystemParameters.RootConstants;
+
+	// Will get filled in with the bindless index in D3DHitGroupSystemParameters.BindlessHitGroupSystemIndexBuffer
+	ByteAddressBuffer HitGroupSystemIndexBuffer;
+
+	// Will get filled in with the bindless index in D3DHitGroupSystemParameters.BindlessHitGroupSystemVertexBuffer
+	ByteAddressBuffer HitGroupSystemVertexBuffer;
+#endif
