- Use non uniform resource index for sampling bindless resources in raytracing shaders
- Re-enable bindless for raytracing

diff --git a/Engine/Config/Windows/BaseWindowsEngine.ini b/Engine/Config/Windows/BaseWindowsEngine.ini
--- a/Engine/Config/Windows/BaseWindowsEngine.ini
+++ b/Engine/Config/Windows/BaseWindowsEngine.ini
@@ -16,8 +16,8 @@ BindlessSamplers=Disabled
 MaxClosuresPerPixel=1
 
 [PCD3D_SM6]
-BindlessResources=Disabled
-BindlessSamplers=Disabled
+BindlessResources=RayTracingOnly
+BindlessSamplers=RayTracingOnly
 
 [SF_VULKAN_SM5]
 BindlessResources=Disabled
diff --git a/Engine/Config/Windows/DataDrivenPlatformInfo.ini b/Engine/Config/Windows/DataDrivenPlatformInfo.ini
--- a/Engine/Config/Windows/DataDrivenPlatformInfo.ini
+++ b/Engine/Config/Windows/DataDrivenPlatformInfo.ini
@@ -122,7 +122,7 @@ bSupportsInstancedStereo=true
 bSupportsPercentageCloserShadows=true
 bSupportsIndexBufferUAVs=true
 BindlessSupport=AllShaderTypes
-StaticShaderBindingLayoutSupport=Unsupported
+StaticShaderBindingLayoutSupport=RayTracingOnly
 bSupportsROV=true
 bSupportsOIT=true
 bSupportsRealTypes=RuntimeDependent
diff --git a/Engine/Shaders/Public/Platform/D3D/D3DCommon.ush b/Engine/Shaders/Public/Platform/D3D/D3DCommon.ush
--- a/Engine/Shaders/Public/Platform/D3D/D3DCommon.ush
+++ b/Engine/Shaders/Public/Platform/D3D/D3DCommon.ush
@@ -145,6 +145,13 @@ void UEReportAssertWithPayload(uint AssertID, uint4 Payload)
 
 #if PLATFORM_SUPPORTS_BINDLESS
 	// SM6.6 path
+#if (RAYCALLABLESHADER || RAYHITGROUPSHADER || RAYGENSHADER || RAYMISSSHADER)
+	// Bindless index can be divergent in hit, miss and callable shaders
+	#define GetResourceFromHeap(Type, Index) ResourceDescriptorHeap[NonUniformResourceIndex(Index)]
+	#define GetSamplerFromHeap(Type, Index) SamplerDescriptorHeap[NonUniformResourceIndex(Index)]
+#else	
 	#define GetResourceFromHeap(Type, Index) ResourceDescriptorHeap[Index]
 	#define GetSamplerFromHeap(Type, Index) SamplerDescriptorHeap[Index]
 #endif
+	
+#endif	
diff --git a/Engine/Shaders/Public/Platform/Vulkan/VulkanCommon.ush b/Engine/Shaders/Public/Platform/Vulkan/VulkanCommon.ush
--- a/Engine/Shaders/Public/Platform/Vulkan/VulkanCommon.ush
+++ b/Engine/Shaders/Public/Platform/Vulkan/VulkanCommon.ush
@@ -48,11 +48,17 @@ void ImageInterlockedAddUInt64(RWTexture2D<UlongType> Dest, uint2 Coord, UlongTy
 	#define VULKAN_SRV_HEAP(Type) VULKAN_HEAP_NAME(BINDLESS_SRV_ARRAY_PREFIX,Type)
 	#define VULKAN_UAV_HEAP(Type) VULKAN_HEAP_NAME(BINDLESS_UAV_ARRAY_PREFIX,Type)
 
+	#if (RAYCALLABLESHADER || RAYHITGROUPSHADER || RAYGENSHADER || RAYMISSSHADER)
+		#define VULKAN_HEAP_ACCESS(BindlessIndex) NonUniformResourceIndex(BindlessIndex)
+	#else
+		#define VULKAN_HEAP_ACCESS(BindlessIndex) BindlessIndex
+	#endif
+
 	#define UB_RESOURCE_MEMBER_SRV(Type, UBName, Name) \
 		typedef Type SafeType##UBName##_##Name; \
 		SafeType##UBName##_##Name VULKAN_SRV_HEAP(SafeType##UBName##_##Name)[]; \
 		Type GetBindlessResource##UBName##_##Name() { \
-			return VULKAN_SRV_HEAP(SafeType##UBName##_##Name)[UB_CB_PREFIXED_MEMBER_ACCESS(UBName, BindlessSRV_, Name)]; \
+			return VULKAN_SRV_HEAP(SafeType##UBName##_##Name)[VULKAN_HEAP_ACCESS(UB_CB_PREFIXED_MEMBER_ACCESS(UBName, BindlessSRV_, Name))]; \
 		} \
 		static const Type UBName##_##Name = GetBindlessResource##UBName##_##Name();
 
@@ -60,7 +66,7 @@ void ImageInterlockedAddUInt64(RWTexture2D<UlongType> Dest, uint2 Coord, UlongTy
 		typedef Type SafeType##UBName##_##Name; \
 		SafeType##UBName##_##Name VULKAN_UAV_HEAP(SafeType##UBName##_##Name)[]; \
 		Type GetBindlessResource##UBName##_##Name() { \
-			return VULKAN_UAV_HEAP(SafeType##UBName##_##Name)[UB_CB_PREFIXED_MEMBER_ACCESS(UBName, BindlessUAV_, Name)]; \
+			return VULKAN_UAV_HEAP(SafeType##UBName##_##Name)[VULKAN_HEAP_ACCESS(UB_CB_PREFIXED_MEMBER_ACCESS(UBName, BindlessUAV_, Name))]; \
 		} \
 		static const Type UBName##_##Name = GetBindlessResource##UBName##_##Name();
 
@@ -68,7 +74,7 @@ void ImageInterlockedAddUInt64(RWTexture2D<UlongType> Dest, uint2 Coord, UlongTy
 		typedef Type SafeType##UBName##_##Name; \
 		SafeType##UBName##_##Name VULKAN_SAMPLER_HEAP(SafeType##UBName##_##Name)[]; \
 		Type GetBindlessSampler##UBName##_##Name() { \
-			return VULKAN_SAMPLER_HEAP(SafeType##UBName##_##Name)[UB_CB_PREFIXED_MEMBER_ACCESS(UBName, BindlessSampler_, Name)]; \
+			return VULKAN_SAMPLER_HEAP(SafeType##UBName##_##Name)[VULKAN_HEAP_ACCESS(UB_CB_PREFIXED_MEMBER_ACCESS(UBName, BindlessSampler_, Name))]; \
 		} \
 		static const Type UBName##_##Name = GetBindlessSampler##UBName##_##Name();
 #endif
diff --git a/Engine/Shaders/Public/ShaderVersion.ush b/Engine/Shaders/Public/ShaderVersion.ush
--- a/Engine/Shaders/Public/ShaderVersion.ush
+++ b/Engine/Shaders/Public/ShaderVersion.ush
@@ -3,5 +3,5 @@
 // Each time the console command is executed it generates a new GUID. As this file is included
 // in Platform.ush (which should be included in any shader) it allows to invalidate the shader DDC.
 // 
-#pragma message("UESHADERMETADATA_VERSION C08D4C8D-F9A6-4473-95C7-99576A30AAAF")
+#pragma message("UESHADERMETADATA_VERSION 226C088E-0B2F-4768-AB17-6C8932D98D03")
 
diff --git a/Engine/Source/Developer/VulkanShaderFormat/Private/VulkanShaderCompiler.cpp b/Engine/Source/Developer/VulkanShaderFormat/Private/VulkanShaderCompiler.cpp
--- a/Engine/Source/Developer/VulkanShaderFormat/Private/VulkanShaderCompiler.cpp
+++ b/Engine/Source/Developer/VulkanShaderFormat/Private/VulkanShaderCompiler.cpp
@@ -233,23 +233,29 @@ struct FVulkanShaderParameterParserPlatformConfiguration : public FShaderParamet
 			EnumAddFlags(Flags, EShaderParameterParserConfigurationFlags::ReplaceGlobals);
 		}
 	}
-
 	virtual FString GenerateBindlessAccess(EBindlessConversionType BindlessType, FStringView FullTypeString, FStringView ArrayNameOverride, FStringView IndexString) const final
 	{
-		if (bIsRayTracingShader && (BindlessType == EBindlessConversionType::SRV))
+		if (bIsRayTracingShader)
 		{
-			// Patch the HitGroupSystemIndexBuffer/HitGroupSystemVertexBuffer indices to use the ones contained in the shader record
-			if (IndexString == HitGroupSystemIndexBufferName)
-			{
-				IndexString = TEXTVIEW("VulkanHitGroupSystemParameters.BindlessHitGroupSystemIndexBuffer");
-			}
-			else if (IndexString == HitGroupSystemVertexBufferName)
+			if (BindlessType == EBindlessConversionType::SRV)
 			{
-				IndexString = TEXTVIEW("VulkanHitGroupSystemParameters.BindlessHitGroupSystemVertexBuffer");
+				// Patch the HitGroupSystemIndexBuffer/HitGroupSystemVertexBuffer indices to use the ones contained in the shader record
+				if (IndexString == HitGroupSystemIndexBufferName)
+				{
+					IndexString = TEXTVIEW("VulkanHitGroupSystemParameters.BindlessHitGroupSystemIndexBuffer");
+				}
+				else if (IndexString == HitGroupSystemVertexBufferName)
+				{
+					IndexString = TEXTVIEW("VulkanHitGroupSystemParameters.BindlessHitGroupSystemVertexBuffer");
+				}
 			}
+						
+			// Raytracing shaders need NonUniformResourceIndex because bindless index can be divergent in hit/miss/callable shaders
+			return FString::Printf(TEXT("%.*s[NonUniformResourceIndex(%.*s)]"),
+				ArrayNameOverride.Len(), ArrayNameOverride.GetData(),
+				IndexString.Len(), IndexString.GetData());
 		}
 
-		// Heap[Index]
 		return FString::Printf(TEXT("%.*s[%.*s]"),
 			ArrayNameOverride.Len(), ArrayNameOverride.GetData(),
 			IndexString.Len(), IndexString.GetData());
diff --git a/Engine/Source/Developer/Windows/ShaderFormatD3D/Private/D3DShaderCompiler.cpp b/Engine/Source/Developer/Windows/ShaderFormatD3D/Private/D3DShaderCompiler.cpp
--- a/Engine/Source/Developer/Windows/ShaderFormatD3D/Private/D3DShaderCompiler.cpp
+++ b/Engine/Source/Developer/Windows/ShaderFormatD3D/Private/D3DShaderCompiler.cpp
@@ -1325,24 +1325,33 @@ struct FD3DShaderParameterParserPlatformConfiguration : public FShaderParameterP
 
 	virtual FString GenerateBindlessAccess(EBindlessConversionType BindlessType, FStringView FullTypeString, FStringView ArrayNameOverride, FStringView IndexString) const final
 	{
-		if (bIsRayTracingShader && (BindlessType == EBindlessConversionType::SRV))
-		{
-			// Patch the HitGroupSystemIndexBuffer/HitGroupSystemVertexBuffer indices to use the ones contained in the shader record
-			if (IndexString == HitGroupSystemIndexBufferName)
-			{
-				IndexString = TEXTVIEW("D3DHitGroupSystemParameters.BindlessHitGroupSystemIndexBuffer");
-			}
-			else if (IndexString == HitGroupSystemVertexBufferName)
-			{
-				IndexString = TEXTVIEW("D3DHitGroupSystemParameters.BindlessHitGroupSystemVertexBuffer");
-			}
-		}
-
 		// GetResourceFromHeap(Type, Index) ResourceDescriptorHeap[Index]
 		// GetSamplerFromHeap(Type, Index)  SamplerDescriptorHeap[Index]
 
 		const TCHAR* HeapString = BindlessType == EBindlessConversionType::Sampler ? TEXT("SamplerDescriptorHeap") : TEXT("ResourceDescriptorHeap");
 
+		if (bIsRayTracingShader)
+		{
+			if (BindlessType == EBindlessConversionType::SRV)
+			{
+				// Patch the HitGroupSystemIndexBuffer/HitGroupSystemVertexBuffer indices to use the ones contained in the shader record
+				if (IndexString == HitGroupSystemIndexBufferName)
+				{
+					IndexString = TEXTVIEW("D3DHitGroupSystemParameters.BindlessHitGroupSystemIndexBuffer");
+				}
+				else if (IndexString == HitGroupSystemVertexBufferName)
+				{
+					IndexString = TEXTVIEW("D3DHitGroupSystemParameters.BindlessHitGroupSystemVertexBuffer");
+				}
+			}
+
+			// Raytracing shaders need NonUniformResourceIndex because bindless index can be divergent in hit/miss/callable shaders
+			return FString::Printf(TEXT("%s[NonUniformResourceIndex(%.*s)]"),
+				HeapString,
+				IndexString.Len(), IndexString.GetData()
+			);
+		}
+			
 		return FString::Printf(TEXT("%s[%.*s]"),
 			HeapString,
 			IndexString.Len(), IndexString.GetData()
