[shaders] emergency version bump, local iteration may have corrupted shared DDC

diff --git a/Engine/Shaders/Public/ShaderVersion.ush b/Engine/Shaders/Public/ShaderVersion.ush
--- a/Engine/Shaders/Public/ShaderVersion.ush
+++ b/Engine/Shaders/Public/ShaderVersion.ush
@@ -3,5 +3,5 @@
 // Each time the console command is executed it generates a new GUID. As this file is included
 // in Platform.ush (which should be included in any shader) it allows to invalidate the shader DDC.
 // 
-#pragma message("UESHADERMETADATA_VERSION 226C088E-0B2F-4768-AB17-6C8932D98D03")
+#pragma message("UESHADERMETADATA_VERSION 45ED0335-634F-4035-BD78-1F5BBE687579")
 
