[shaders] version bump to clear potential bad ddc data (again)

#rnx

diff --git a/Engine/Shaders/Public/ShaderVersion.ush b/Engine/Shaders/Public/ShaderVersion.ush
--- a/Engine/Shaders/Public/ShaderVersion.ush
+++ b/Engine/Shaders/Public/ShaderVersion.ush
@@ -3,5 +3,5 @@
 // Each time the console command is executed it generates a new GUID. As this file is included
 // in Platform.ush (which should be included in any shader) it allows to invalidate the shader DDC.
 // 
-#pragma message("UESHADERMETADATA_VERSION 45ED0335-634F-4035-BD78-1F5BBE687579")
+#pragma message("UESHADERMETADATA_VERSION 8BBB7333-116A-407C-A7B2-F7D40D5B6ADE")
 
