Substrate - fixed non written pixels and unlit pixels to return Unlit shading model ID and color matching the legacy behavior.

#rnx
#tests pc dx12 editor substrate=1

diff --git a/Engine/Shaders/Public/SubstratePublic.ush b/Engine/Shaders/Public/SubstratePublic.ush
--- a/Engine/Shaders/Public/SubstratePublic.ush
+++ b/Engine/Shaders/Public/SubstratePublic.ush
@@ -68,8 +68,8 @@ FSubstrateGBuffer SubstratePublic_GetSubstrateGBufferFromFirstBSDF(uint2 PixelCo
 		SGBuffer.WorldNormal		= SubstrateGetWorldNormal(SubstratePixelHeader, BSDF, SubstrateAddressing);
 		SGBuffer.Roughness			= SubstrateGetBSDFRoughness(BSDF);
 		SGBuffer.MaterialAO			= SubstrateGetAO(SubstratePixelHeader);
-		SGBuffer.ShadingModelColor	= GetShadingModelColor(SubstrateGetLegacyShadingModels(BSDF));
 		SGBuffer.ShadingModelID		= SubstrateGetLegacyShadingModels(BSDF);
+		SGBuffer.ShadingModelColor	= GetShadingModelColor(SGBuffer.ShadingModelID);
 		SGBuffer.StoredBaseColor	= SubstrateGetBSDFBaseColor(BSDF);
 		SGBuffer.StoredSpecular		= SubstrateGetBSDFSpecular(BSDF).rrr;
 		SGBuffer.WorldTangent		= SubstrateGetWorldTangent(SubstratePixelHeader, BSDF, SubstrateAddressing);
@@ -105,6 +105,12 @@ FSubstrateGBuffer SubstratePublic_GetSubstrateGBufferFromFirstBSDF(uint2 PixelCo
 			SGBuffer.Opacity = HAIR_BACKLIT(BSDF);			// BackLit
 		}
 	}
+	else if (SubstratePixelHeader.State == 0)				// Only Unlit and non written pixels have a state of 0. Those pixels should show up as unlit to respect the legacy behavior.
+	{		
+		SGBuffer.bIsValid			= true;
+		SGBuffer.ShadingModelID		= SHADINGMODELID_UNLIT;
+		SGBuffer.ShadingModelColor	= GetShadingModelColor(SGBuffer.ShadingModelID);
+	}
 #endif
 
 	return SGBuffer;
