Fix Apple Metal SM5 & SM6 corruption by introducing a depth bias in prepass that use OUTPUT_PIXEL_DEPTH_OFFSET
This is caused by a precision issue when mixing vertex and pixel shader depths and seemingly not guaranteeing SV_DepthLessEqual

#tests Tested in CitySample and TM-Decals from QAGame

diff --git a/Engine/Shaders/Private/DepthOnlyPixelShader.usf b/Engine/Shaders/Private/DepthOnlyPixelShader.usf
--- a/Engine/Shaders/Private/DepthOnlyPixelShader.usf
+++ b/Engine/Shaders/Private/DepthOnlyPixelShader.usf
@@ -55,6 +55,10 @@ void Main(
 		{
 			ApplyPixelDepthOffsetToMaterialParameters(MaterialParameters, PixelMaterialInputs, OutDepth);
 		}
+
+		#if APPLE_DEPTH_BIAS_HACK
+		OutDepth -= APPLE_DEPTH_BIAS_VALUE;
+		#endif
 	#endif
 
 
diff --git a/Engine/Shaders/Private/VelocityShader.usf b/Engine/Shaders/Private/VelocityShader.usf
--- a/Engine/Shaders/Private/VelocityShader.usf
+++ b/Engine/Shaders/Private/VelocityShader.usf
@@ -137,6 +137,10 @@ void MainPixelShader(
 
 #if OUTPUT_PIXEL_DEPTH_OFFSET
 	ApplyPixelDepthOffsetToMaterialParameters(MaterialParameters, PixelMaterialInputs, OutDepth);
+
+	#if APPLE_DEPTH_BIAS_HACK
+	OutDepth -= APPLE_DEPTH_BIAS_VALUE;
+	#endif
 #endif
 
 #if MATERIALBLENDING_MASKED_USING_COVERAGE
diff --git a/Engine/Shaders/Public/Platform/Metal/MetalCommon.ush b/Engine/Shaders/Public/Platform/Metal/MetalCommon.ush
--- a/Engine/Shaders/Public/Platform/Metal/MetalCommon.ush
+++ b/Engine/Shaders/Public/Platform/Metal/MetalCommon.ush
@@ -26,6 +26,13 @@
 	#define USE_FORCE_TEXTURE_MIP 1
 #endif
 
+// Apple needs a bias because of a precision issue that can't guarantee SV_DepthLessEqual
+// This is causing base pass to partially fail depth test when depth prepass use OUTPUT_PIXEL_DEPTH_OFFSET
+#if METAL_SM5_PROFILE || METAL_SM6_PROFILE
+	#define APPLE_DEPTH_BIAS_HACK 1
+	#define APPLE_DEPTH_BIAS_VALUE 0.000001
+#endif
+
 /* Too many bugs to enable wave-intrinsics just yet, but we are getting there.
 #if COMPUTESHADER || PIXELSHADER
 #define COMPILER_SUPPORTS_WAVE_ONCE 1
