Unused interpolators feature added on PS5 from VS->PS pipelines based on shader reflection - it already spits out only the used interpolators.

diff --git a/Engine/Shaders/Public/ShaderVersion.ush b/Engine/Shaders/Public/ShaderVersion.ush
--- a/Engine/Shaders/Public/ShaderVersion.ush
+++ b/Engine/Shaders/Public/ShaderVersion.ush
@@ -3,4 +3,4 @@
 // Each time the console command is executed it generates a new GUID. As this file is included
 // in Platform.ush (which should be included in any shader) it allows to invalidate the shader DDC.
 // 
-#pragma message("UESHADERMETADATA_VERSION AAF7C14D-1BF2-4F8E-A286-C910215B5C9E")
+#pragma message("UESHADERMETADATA_VERSION 3E78E5FC-CD27-4749-8B35-BA0F068BD837")
\ No newline at end of file
