Metal - Apply depth bias on iOS & tvOS for SM5

#tests Local test on iPhone 16 pro with and without the bias. Also tested ES3_1 but PixelDepthOffset is not applied on ES3_1

diff --git a/Engine/Shaders/Public/Platform/Metal/MetalCommon.ush b/Engine/Shaders/Public/Platform/Metal/MetalCommon.ush
--- a/Engine/Shaders/Public/Platform/Metal/MetalCommon.ush
+++ b/Engine/Shaders/Public/Platform/Metal/MetalCommon.ush
@@ -28,7 +28,7 @@
 
 // Apple needs a bias because of a precision issue that can't guarantee SV_DepthLessEqual
 // This is causing base pass to partially fail depth test when depth prepass use OUTPUT_PIXEL_DEPTH_OFFSET
-#if METAL_SM5_PROFILE || METAL_SM6_PROFILE
+#if METAL_SM5_PROFILE || METAL_SM6_PROFILE || METAL_MRT_PROFILE
 	#define APPLE_DEPTH_BIAS_HACK 1
 	#define APPLE_DEPTH_BIAS_VALUE 0.000001
 #endif
