Bindless shader changes
- Splitting Bindless "Resource" into "SRV" and "UAV" to support platforms where that separation is necessary.
- Making Vulkan-style bindless heaps part of the core bindless uniform buffer macros.

diff --git a/Engine/Shaders/Private/MaterialCache/MaterialCacheUnwrapShade.usf b/Engine/Shaders/Private/MaterialCache/MaterialCacheUnwrapShade.usf
--- a/Engine/Shaders/Private/MaterialCache/MaterialCacheUnwrapShade.usf
+++ b/Engine/Shaders/Private/MaterialCache/MaterialCacheUnwrapShade.usf
@@ -87,7 +87,7 @@ void Main(uint3 DTid : SV_DispatchThreadID)
 	const uint2 PhysicalPixelPos  = BinData.PhysicalCoordinate + PagePixelPos;
 	
 	// Always uniform addressing
-	RWTexture2D<float4> PhysicalTarget = GetResourceFromHeap(RWTexture2D<float4>, BinData.PhysicalBindlessHandle);
+	RWTexture2D<float4> PhysicalTarget = GetUAVFromHeap(RWTexture2D<float4>, BinData.PhysicalBindlessHandle);
 
 	// Shader against the current layer
 	float4 ShadedAttribute = ShadeAttributes(BinData, PagePixelPos);
diff --git a/Engine/Shaders/Private/ShaderBundleWorkGraphDispatch.usf b/Engine/Shaders/Private/ShaderBundleWorkGraphDispatch.usf
--- a/Engine/Shaders/Private/ShaderBundleWorkGraphDispatch.usf
+++ b/Engine/Shaders/Private/ShaderBundleWorkGraphDispatch.usf
@@ -44,7 +44,7 @@ void WorkGraphMainCS(
 
 #if UE_BINDLESS_ENABLED
 		const uint RecordArgBufferHandle = InputRecord.Get().PlatformData.z;
-		ByteAddressBuffer BindlessRecordArgBuffer = GetResourceFromHeap(ByteAddressBuffer, RecordArgBufferHandle);
+		ByteAddressBuffer BindlessRecordArgBuffer = GetSRVFromHeap(ByteAddressBuffer, RecordArgBufferHandle);
 		DispatchGridSize = BindlessRecordArgBuffer.Load(ArgOffset);
 #else
 		DispatchGridSize = RecordArgBuffer.Load(ArgOffset);
diff --git a/Engine/Shaders/Public/BindlessResources.ush b/Engine/Shaders/Public/BindlessResources.ush
--- a/Engine/Shaders/Public/BindlessResources.ush
+++ b/Engine/Shaders/Public/BindlessResources.ush
@@ -4,32 +4,40 @@
 
 #include "Platform.ush"
 
-#if defined(UB_RESOURCE_MEMBER_SRV) || defined(UB_RESOURCE_MEMBER_UAV)
-	// Using platform defined bindless macros
+#if UE_BINDLESS_ENABLED
+	#if !defined(UE_BINDLESS_SRV_HEAP_DECL)
+		#define UE_BINDLESS_SRV_HEAP_DECL(Type)
+	#endif
+
+	#if !defined(UE_BINDLESS_UAV_HEAP_DECL)
+		#define UE_BINDLESS_UAV_HEAP_DECL(Type)
+	#endif
+
+	#if !defined(UE_BINDLESS_SAMPLER_HEAP_DECL)
+		#define UE_BINDLESS_SAMPLER_HEAP_DECL(Type)
+	#endif
 
-#elif UE_BINDLESS_ENABLED
+	// Using platform defined bindless macros
 	#define UB_RESOURCE_MEMBER_SRV(Type, UBName, Name) \
 		typedef Type SafeType##UBName##_##Name; \
-		SafeType##UBName##_##Name GetBindlessResource##UBName##_##Name() { return GetResourceFromHeap(SafeType##UBName##_##Name, UB_CB_PREFIXED_MEMBER_ACCESS(UBName, BindlessSRV_, Name)); } \
+		UE_BINDLESS_SRV_HEAP_DECL(SafeType##UBName##_##Name); \
+		SafeType##UBName##_##Name GetBindlessResource##UBName##_##Name() { return GetSRVFromHeap(SafeType##UBName##_##Name, UB_CB_PREFIXED_MEMBER_ACCESS(UBName, BindlessSRV_, Name)); } \
 		static const SafeType##UBName##_##Name UBName##_##Name = GetBindlessResource##UBName##_##Name();
 
 	#define UB_RESOURCE_MEMBER_UAV(Type, UBName, Name) \
 		typedef Type SafeType##UBName##_##Name; \
-		SafeType##UBName##_##Name GetBindlessResource##UBName##_##Name() { return GetResourceFromHeap(SafeType##UBName##_##Name, UB_CB_PREFIXED_MEMBER_ACCESS(UBName, BindlessUAV_, Name)); } \
+		UE_BINDLESS_UAV_HEAP_DECL(SafeType##UBName##_##Name); \
+		SafeType##UBName##_##Name GetBindlessResource##UBName##_##Name() { return GetUAVFromHeap(SafeType##UBName##_##Name, UB_CB_PREFIXED_MEMBER_ACCESS(UBName, BindlessUAV_, Name)); } \
 		static const SafeType##UBName##_##Name UBName##_##Name = GetBindlessResource##UBName##_##Name();
-#else
-	#define UB_RESOURCE_MEMBER_SRV(Type, UBName, Name) Type UBName##_##Name
-	#define UB_RESOURCE_MEMBER_UAV(Type, UBName, Name) Type UBName##_##Name
-#endif
-
-#if defined(UB_RESOURCE_MEMBER_SAMPLER)
-	// Using platform defined bindless macros
 
-#elif UE_BINDLESS_ENABLED
 	#define UB_RESOURCE_MEMBER_SAMPLER(Type, UBName, Name) \
-		Type GetBindlessSampler##UBName##_##Name() { return GetSamplerFromHeap(Type, UB_CB_PREFIXED_MEMBER_ACCESS(UBName, BindlessSampler_, Name)); } \
-		static const Type UBName##_##Name = GetBindlessSampler##UBName##_##Name();
+		typedef Type SafeType##UBName##_##Name; \
+		UE_BINDLESS_SAMPLER_HEAP_DECL(SafeType##UBName##_##Name); \
+		SafeType##UBName##_##Name GetBindlessSampler##UBName##_##Name() { return GetSamplerFromHeap(SafeType##UBName##_##Name, UB_CB_PREFIXED_MEMBER_ACCESS(UBName, BindlessSampler_, Name)); } \
+		static const SafeType##UBName##_##Name UBName##_##Name = GetBindlessSampler##UBName##_##Name();
 #else
+	#define UB_RESOURCE_MEMBER_SRV(Type, UBName, Name) Type UBName##_##Name
+	#define UB_RESOURCE_MEMBER_UAV(Type, UBName, Name) Type UBName##_##Name
 	#define UB_RESOURCE_MEMBER_SAMPLER(Type, UBName, Name) Type UBName##_##Name
 #endif
 
@@ -66,7 +74,7 @@ uint GetResourceIndexFromResourceCollection(FResourceCollection ResourceCollecti
 		TType TextureFromCollection_##TType(FResourceCollection ResourceCollection, uint CollectionIndex) \
 		{ \
 			uint ResourceIndex = GetResourceIndexFromResourceCollection(ResourceCollection, CollectionIndex); \
-			return GetResourceFromHeap(TType, ResourceIndex); \
+			return GetSRVFromHeap(TType, ResourceIndex); \
 		}
 #endif
 
diff --git a/Engine/Shaders/Public/Platform/D3D/D3DCommon.ush b/Engine/Shaders/Public/Platform/D3D/D3DCommon.ush
--- a/Engine/Shaders/Public/Platform/D3D/D3DCommon.ush
+++ b/Engine/Shaders/Public/Platform/D3D/D3DCommon.ush
@@ -144,14 +144,14 @@ void UEReportAssertWithPayload(uint AssertID, uint4 Payload)
 #endif
 
 #if PLATFORM_SUPPORTS_BINDLESS
-	// SM6.6 path
-#if (RAYCALLABLESHADER || RAYHITGROUPSHADER || RAYGENSHADER || RAYMISSSHADER)
-	// Bindless index can be divergent in hit, miss and callable shaders
-	#define GetResourceFromHeap(Type, Index) ResourceDescriptorHeap[NonUniformResourceIndex(Index)]
-	#define GetSamplerFromHeap(Type, Index) SamplerDescriptorHeap[NonUniformResourceIndex(Index)]
-#else	
-	#define GetResourceFromHeap(Type, Index) ResourceDescriptorHeap[Index]
-	#define GetSamplerFromHeap(Type, Index) SamplerDescriptorHeap[Index]
-#endif
-	
+	#if (RAYCALLABLESHADER || RAYHITGROUPSHADER || RAYGENSHADER || RAYMISSSHADER)
+		// Bindless index can be divergent in hit, miss and callable shaders
+		#define UE_BINDLESS_INDEX(Index) NonUniformResourceIndex(Index)
+	#else
+		#define UE_BINDLESS_INDEX(Index) Index
+	#endif
+
+	#define GetSRVFromHeap(Type, Index) ResourceDescriptorHeap[UE_BINDLESS_INDEX(Index)]
+	#define GetUAVFromHeap(Type, Index) ResourceDescriptorHeap[UE_BINDLESS_INDEX(Index)]
+	#define GetSamplerFromHeap(Type, Index) SamplerDescriptorHeap[UE_BINDLESS_INDEX(Index)]
 #endif	
diff --git a/Engine/Shaders/Public/Platform/Metal/MetalCommon.ush b/Engine/Shaders/Public/Platform/Metal/MetalCommon.ush
--- a/Engine/Shaders/Public/Platform/Metal/MetalCommon.ush
+++ b/Engine/Shaders/Public/Platform/Metal/MetalCommon.ush
@@ -113,7 +113,8 @@
     #endif
     
     #if PLATFORM_SUPPORTS_BINDLESS
-        #define GetResourceFromHeap(Type, Index) ResourceDescriptorHeap[Index]
+        #define GetSRVFromHeap(Type, Index) ResourceDescriptorHeap[Index]
+        #define GetUAVFromHeap(Type, Index) ResourceDescriptorHeap[Index]
         #define GetSamplerFromHeap(Type, Index) SamplerDescriptorHeap[Index]
     #endif
 #endif
diff --git a/Engine/Shaders/Public/Platform/Vulkan/VulkanCommon.ush b/Engine/Shaders/Public/Platform/Vulkan/VulkanCommon.ush
--- a/Engine/Shaders/Public/Platform/Vulkan/VulkanCommon.ush
+++ b/Engine/Shaders/Public/Platform/Vulkan/VulkanCommon.ush
@@ -40,7 +40,7 @@ void ImageInterlockedAddUInt64(RWTexture2D<UlongType> Dest, uint2 Coord, UlongTy
 }
 #endif
 
-#if PLATFORM_SUPPORTS_BINDLESS
+#if UE_BINDLESS_ENABLED
 	// While we use descriptor indexing, declare unbounded resource arrays
 	#define VULKAN_HEAP_NAME_INNER(Prefix, Type) Prefix##Type
 	#define VULKAN_HEAP_NAME(Prefix, Type) VULKAN_HEAP_NAME_INNER(Prefix, Type)
@@ -54,29 +54,13 @@ void ImageInterlockedAddUInt64(RWTexture2D<UlongType> Dest, uint2 Coord, UlongTy
 		#define VULKAN_HEAP_ACCESS(BindlessIndex) BindlessIndex
 	#endif
 
-	#define UB_RESOURCE_MEMBER_SRV(Type, UBName, Name) \
-		typedef Type SafeType##UBName##_##Name; \
-		SafeType##UBName##_##Name VULKAN_SRV_HEAP(SafeType##UBName##_##Name)[]; \
-		Type GetBindlessResource##UBName##_##Name() { \
-			return VULKAN_SRV_HEAP(SafeType##UBName##_##Name)[VULKAN_HEAP_ACCESS(UB_CB_PREFIXED_MEMBER_ACCESS(UBName, BindlessSRV_, Name))]; \
-		} \
-		static const Type UBName##_##Name = GetBindlessResource##UBName##_##Name();
-
-	#define UB_RESOURCE_MEMBER_UAV(Type, UBName, Name) \
-		typedef Type SafeType##UBName##_##Name; \
-		SafeType##UBName##_##Name VULKAN_UAV_HEAP(SafeType##UBName##_##Name)[]; \
-		Type GetBindlessResource##UBName##_##Name() { \
-			return VULKAN_UAV_HEAP(SafeType##UBName##_##Name)[VULKAN_HEAP_ACCESS(UB_CB_PREFIXED_MEMBER_ACCESS(UBName, BindlessUAV_, Name))]; \
-		} \
-		static const Type UBName##_##Name = GetBindlessResource##UBName##_##Name();
-
-	#define UB_RESOURCE_MEMBER_SAMPLER(Type, UBName, Name) \
-		typedef Type SafeType##UBName##_##Name; \
-		SafeType##UBName##_##Name VULKAN_SAMPLER_HEAP(SafeType##UBName##_##Name)[]; \
-		Type GetBindlessSampler##UBName##_##Name() { \
-			return VULKAN_SAMPLER_HEAP(SafeType##UBName##_##Name)[VULKAN_HEAP_ACCESS(UB_CB_PREFIXED_MEMBER_ACCESS(UBName, BindlessSampler_, Name))]; \
-		} \
-		static const Type UBName##_##Name = GetBindlessSampler##UBName##_##Name();
+	#define UE_BINDLESS_SRV_HEAP_DECL(Type) Type VULKAN_SRV_HEAP(Type)[];
+	#define UE_BINDLESS_UAV_HEAP_DECL(Type) Type VULKAN_UAV_HEAP(Type)[];
+	#define UE_BINDLESS_SAMPLER_HEAP_DECL(Type) Type VULKAN_SAMPLER_HEAP(Type)[];
+
+	#define GetSRVFromHeap(Type, Index) VULKAN_SRV_HEAP(Type)[VULKAN_HEAP_ACCESS(Index)]
+	#define GetUAVFromHeap(Type, Index) VULKAN_UAV_HEAP(Type)[VULKAN_HEAP_ACCESS(Index)]
+	#define GetSamplerFromHeap(Type, Index) VULKAN_SAMPLER_HEAP(Type)[VULKAN_HEAP_ACCESS(Index)]
 #endif
 
 #if VULKAN_SUPPORTS_SUBGROUP_SIZE_CONTROL
@@ -88,7 +72,7 @@ void ImageInterlockedAddUInt64(RWTexture2D<UlongType> Dest, uint2 Coord, UlongTy
 	#define WAVESIZE(N) VULKAN_WAVESIZE(N)
 #endif
 
-#if PLATFORM_SUPPORTS_BINDLESS && (RAYCALLABLESHADER || RAYHITGROUPSHADER || RAYGENSHADER || RAYMISSSHADER)
+#if UE_BINDLESS_ENABLED && (RAYCALLABLESHADER || RAYHITGROUPSHADER || RAYGENSHADER || RAYMISSSHADER)
 	#define OVERRIDE_RAY_TRACING_HIT_GROUP_SYSTEM_RESOURCES 1
 	#include "/Engine/Shared/RayTracingBuiltInResources.h"
 
@@ -149,7 +133,7 @@ void ImageInterlockedAddUInt64(RWTexture2D<UlongType> Dest, uint2 Coord, UlongTy
 // Must match GetGlobalSamplerType() in SpirVShaderCompiler.inl
 // and EGlobalSamplerType in VulkanShaderResources.h
 
-#if PLATFORM_SUPPORTS_BINDLESS
+#if UE_BINDLESS_ENABLED
 	typedef SamplerState SafeTypeGlobalSamplers;
 	SafeTypeGlobalSamplers VULKAN_SAMPLER_HEAP(GlobalSamplers)[];
 	static const SamplerState VulkanGlobalPointClampedSampler     = VULKAN_SAMPLER_HEAP(GlobalSamplers)[1];
diff --git a/Engine/Source/Developer/Apple/MetalShaderFormat/Private/MetalCompileShaderMSC.cpp b/Engine/Source/Developer/Apple/MetalShaderFormat/Private/MetalCompileShaderMSC.cpp
--- a/Engine/Source/Developer/Apple/MetalShaderFormat/Private/MetalCompileShaderMSC.cpp
+++ b/Engine/Source/Developer/Apple/MetalShaderFormat/Private/MetalCompileShaderMSC.cpp
@@ -525,7 +525,8 @@ struct FMetalShaderParameterParserPlatformConfiguration : public FShaderParamete
 
 	virtual FString GenerateBindlessAccess(EBindlessConversionType BindlessType, FStringView FullTypeString, FStringView ArrayNameOverride, FStringView IndexString) const final
 	{
-		// GetResourceFromHeap(Type, Index) ResourceDescriptorHeap[Index]
+		// GetSRVFromHeap(Type, Index) ResourceDescriptorHeap[Index]
+		// GetUAVFromHeap(Type, Index) ResourceDescriptorHeap[Index]
 		// GetSamplerFromHeap(Type, Index)  SamplerDescriptorHeap[Index]
 
 		const TCHAR* HeapString = BindlessType == EBindlessConversionType::Sampler ? TEXT("SamplerDescriptorHeap") : TEXT("ResourceDescriptorHeap");
diff --git a/Engine/Source/Developer/VulkanShaderFormat/Private/VulkanShaderCompiler.cpp b/Engine/Source/Developer/VulkanShaderFormat/Private/VulkanShaderCompiler.cpp
--- a/Engine/Source/Developer/VulkanShaderFormat/Private/VulkanShaderCompiler.cpp
+++ b/Engine/Source/Developer/VulkanShaderFormat/Private/VulkanShaderCompiler.cpp
@@ -235,6 +235,10 @@ struct FVulkanShaderParameterParserPlatformConfiguration : public FShaderParamet
 	}
 	virtual FString GenerateBindlessAccess(EBindlessConversionType BindlessType, FStringView FullTypeString, FStringView ArrayNameOverride, FStringView IndexString) const final
 	{
+		// GetSRVFromHeap(Type, Index) VULKAN_SRV_HEAP(Type)[VULKAN_HEAP_ACCESS(Index)]
+		// GetUAVFromHeap(Type, Index) VULKAN_UAV_HEAP(Type)[VULKAN_HEAP_ACCESS(Index)]
+		// GetSamplerFromHeap(Type, Index) VULKAN_SAMPLER_HEAP(Type)[VULKAN_HEAP_ACCESS(Index)]
+
 		if (bIsRayTracingShader)
 		{
 			if (BindlessType == EBindlessConversionType::SRV)
@@ -249,7 +253,7 @@ struct FVulkanShaderParameterParserPlatformConfiguration : public FShaderParamet
 					IndexString = TEXTVIEW("VulkanHitGroupSystemParameters.BindlessHitGroupSystemVertexBuffer");
 				}
 			}
-						
+
 			// Raytracing shaders need NonUniformResourceIndex because bindless index can be divergent in hit/miss/callable shaders
 			return FString::Printf(TEXT("%.*s[NonUniformResourceIndex(%.*s)]"),
 				ArrayNameOverride.Len(), ArrayNameOverride.GetData(),
diff --git a/Engine/Source/Developer/Windows/ShaderFormatD3D/Private/D3DShaderCompiler.cpp b/Engine/Source/Developer/Windows/ShaderFormatD3D/Private/D3DShaderCompiler.cpp
--- a/Engine/Source/Developer/Windows/ShaderFormatD3D/Private/D3DShaderCompiler.cpp
+++ b/Engine/Source/Developer/Windows/ShaderFormatD3D/Private/D3DShaderCompiler.cpp
@@ -1395,7 +1395,8 @@ struct FD3DShaderParameterParserPlatformConfiguration : public FShaderParameterP
 
 	virtual FString GenerateBindlessAccess(EBindlessConversionType BindlessType, FStringView FullTypeString, FStringView ArrayNameOverride, FStringView IndexString) const final
 	{
-		// GetResourceFromHeap(Type, Index) ResourceDescriptorHeap[Index]
+		// GetSRVFromHeap(Type, Index) ResourceDescriptorHeap[Index]
+		// GetUAVFromHeap(Type, Index) ResourceDescriptorHeap[Index]
 		// GetSamplerFromHeap(Type, Index)  SamplerDescriptorHeap[Index]
 
 		const TCHAR* HeapString = BindlessType == EBindlessConversionType::Sampler ? TEXT("SamplerDescriptorHeap") : TEXT("ResourceDescriptorHeap");
diff --git a/Engine/Source/Runtime/RenderCore/Private/ShaderParameterParser.cpp b/Engine/Source/Runtime/RenderCore/Private/ShaderParameterParser.cpp
--- a/Engine/Source/Runtime/RenderCore/Private/ShaderParameterParser.cpp
+++ b/Engine/Source/Runtime/RenderCore/Private/ShaderParameterParser.cpp
@@ -938,6 +938,7 @@ static FStringView GetBindlessArrayHeapPrefix(EBindlessConversionType InConversi
 FString FShaderParameterParser::GenerateBindlessParameterDeclaration(const FParsedShaderParameter& ParsedParameter) const
 {
 	const bool bIsSampler = (ParsedParameter.BindlessConversionType == EBindlessConversionType::Sampler);
+	const bool bIsUAV = (ParsedParameter.BindlessConversionType == EBindlessConversionType::UAV);
 	const FStringView Name = ParsedParameter.ParsedName;
 	const FStringView Type = ParsedParameter.ParsedType;
 	const TCHAR* StorageClass = ParsedParameter.bGloballyCoherent ? TEXT("globallycoherent ") : TEXT("");
@@ -982,9 +983,10 @@ FString FShaderParameterParser::GenerateBindlessParameterDeclaration(const FPars
 	}
 
 	const FString BindlessAccess = PlatformConfiguration.GenerateBindlessAccess(ParsedParameter.BindlessConversionType, FullType, HeapName, IndexString);
-	const TCHAR* FunctionPrefix = bIsSampler ? TEXT("GetBindlessSampler") : TEXT("GetBindlessResource");
+	const TCHAR* FunctionPrefix = bIsSampler ? TEXT("GetBindlessSampler") : bIsUAV ? TEXT("GetBindlessUAV") : TEXT("GetBindlessSRV");
 
-	// e.g. `Type GetBindlessResource##Name() { return GetResourceFromHeap(Type, BindlessResource_##Name); } static const Type Name = GetBindlessResource##Name()`
+	// e.g. `Type GetBindlessSRV##Name() { return GetSRVFromHeap(Type, BindlessSRV_##Name); } static const Type Name = GetBindlessSRV##Name()`
+	// e.g. `Type GetBindlessUAV##Name() { return GetUAVFromHeap(Type, BindlessUAV_##Name); } static const Type Name = GetBindlessUAV##Name()`
 	// or   `Type GetBindlessSampler##Name() { return GetSamplerFromHeap(Type, BindlessSampler_##Name); } static const Type Name = GetBindlessSampler##Name()`
 	Result << FullType << TEXT(" ") << FunctionPrefix << Name << TEXT("() { return ") << BindlessAccess << TEXT("; } ");
 	Result << TEXT("static const ") << FullType << TEXT(" ") << Name << TEXT(" = ") << FunctionPrefix << Name << TEXT("();");
