Constant Buffer Changes
- PLATFORM_SUPPORTS_CONSTANTBUFFER_OBJECT no longer forces all UBs to use ConstantBuffer<> since that causes issues with reflection on some platforms which we need for resources.
- Uniform Buffers now opt-in to using ConstantBuffer<> when the UB does not contain any resources.
- Stable constant buffers now emit their register and space on relevant platforms.
- Adding FShaderCodeGlobalUniformBufferConfiguration for platforms that don't always have constant buffers assigned to registers in the 'expected' order.

diff --git a/Engine/Shaders/Public/Platform.ush b/Engine/Shaders/Public/Platform.ush
--- a/Engine/Shaders/Public/Platform.ush
+++ b/Engine/Shaders/Public/Platform.ush
@@ -517,30 +517,53 @@
 #endif
 #define PLATFORM_MAX_UNIFORM_BUFFER_RANGE_FLOAT4 (PLATFORM_MAX_UNIFORM_BUFFER_RANGE/16u)
 
+// These macros are used by CreateHLSLUniformBufferDeclaration and the functions it calls to generate shader code for Uniform Buffers.
+// Depending on certain scenarios, the UB_CB_ prefixed macros can be swapped out for the UB_CBO_ versions to allow some platforms to define Uniform Buffers differently
+
+#define UB_CB_DEFINITION_START(      UBName                    ) cbuffer UBName {
+#define UB_CB_DEFINITION_START_REG(  UBName, Register, Space   ) cbuffer UBName : register(b##Register, space##Space) {
+#define UB_CB_MEMBER_NAME(           UBName,         MemberName) UBName##_##MemberName
+#define UB_CB_PREFIXED_MEMBER_NAME(  UBName, Prefix, MemberName) Prefix##UBName##_##MemberName
+#define UB_CB_DEFINITION_END(        UBName                    ) }
+#define UB_CB_DEFINITION_END_REG(    UBName, Register, Space   ) UB_CB_DEFINITION_END(UBName)
+#define UB_CB_MEMBER_ACCESS(         UBName,         MemberName) UBName##_##MemberName
+#define UB_CB_PREFIXED_MEMBER_ACCESS(UBName, Prefix, MemberName) Prefix##UBName##_##MemberName
+
+#define UB_CB_DECL_PARAMETER(UBName, StructName, GlobalName) UBName##.##StructName = UBName##_##GlobalName
+#define UB_CB_DECL_RESOURCE(UBName, StructName, GlobalName)  UBName##.##StructName = UBName##_##GlobalName
+
+#define UB_CB_UNIFORM_BLOCK(UBName, MemberName) UB_FLOAT(4) UB_CB_MEMBER_NAME(UBName,MemberName)[PLATFORM_MAX_UNIFORM_BUFFER_RANGE_FLOAT4]
+
 #if PLATFORM_SUPPORTS_CONSTANTBUFFER_OBJECT
-	#define UB_CB_DEFINITION_START(      UBName                    ) struct F##UBName##Constants {
-	#define UB_CB_MEMBER_NAME(           UBName,         MemberName) MemberName
-	#define UB_CB_PREFIXED_MEMBER_NAME(  UBName, Prefix, MemberName) Prefix##MemberName
-	#define UB_CB_DEFINITION_END(        UBName                    ) }; ConstantBuffer<F##UBName##Constants> UBName;
-	#define UB_CB_MEMBER_ACCESS(         UBName,         MemberName) UBName##.##MemberName
-	#define UB_CB_PREFIXED_MEMBER_ACCESS(UBName, Prefix, MemberName) UBName##.##Prefix##MemberName
-
-	#define UB_DECL_PARAMETER(UBName, StructName, GlobalName) UBName##.##StructName = UBName##.##GlobalName
-	#define UB_DECL_RESOURCE(UBName, StructName, GlobalName)  UBName##.##StructName = UBName##_##GlobalName
+	#define UB_CBO_DEFINITION_START(      UBName                    ) struct F##UBName##Constants {
+	#define UB_CBO_DEFINITION_START_REG(  UBName, Register, Space   ) UB_CBO_DEFINITION_START(UBName)
+	#define UB_CBO_MEMBER_NAME(           UBName,         MemberName) MemberName
+	#define UB_CBO_PREFIXED_MEMBER_NAME(  UBName, Prefix, MemberName) Prefix##MemberName
+	#define UB_CBO_DEFINITION_END(        UBName                    ) }; ConstantBuffer<F##UBName##Constants> UBName;
+	#define UB_CBO_DEFINITION_END_REG(    UBName, Register, Space   ) }; ConstantBuffer<F##UBName##Constants> UBName : register(b##Register, space##Space);
+	#define UB_CBO_MEMBER_ACCESS(         UBName,         MemberName) UBName##.##MemberName
+	#define UB_CBO_PREFIXED_MEMBER_ACCESS(UBName, Prefix, MemberName) UBName##.##Prefix##MemberName
+
+	#define UB_CBO_DECL_PARAMETER(UBName, StructName, GlobalName) UBName##.##StructName = UBName##.##GlobalName
+	#define UB_CBO_DECL_RESOURCE(UBName, StructName, GlobalName)  UBName##.##StructName = UBName##_##GlobalName
+
+	#define UB_CBO_UNIFORM_BLOCK(UBName, MemberName) UB_FLOAT(4) UB_CBO_MEMBER_NAME(UBName,MemberName)[PLATFORM_MAX_UNIFORM_BUFFER_RANGE_FLOAT4]
+
 #else
-	#define UB_CB_DEFINITION_START(      UBName                    ) cbuffer UBName {
-	#define UB_STATIC_CB_DEFINITION_START(UBName, ResourceIndex, Space) cbuffer UBName : register(b##ResourceIndex, space##Space) {
-	#define UB_CB_MEMBER_NAME(           UBName,         MemberName) UBName##_##MemberName
-	#define UB_CB_PREFIXED_MEMBER_NAME(  UBName, Prefix, MemberName) Prefix##UBName##_##MemberName
-	#define UB_CB_DEFINITION_END(        UBName                    ) }
-	#define UB_CB_MEMBER_ACCESS(         UBName,         MemberName) UBName##_##MemberName
-	#define UB_CB_PREFIXED_MEMBER_ACCESS(UBName, Prefix, MemberName) Prefix##UBName##_##MemberName
+	#define UB_CBO_DEFINITION_START(      UBName                    ) UB_CB_DEFINITION_START(UBName)
+	#define UB_CBO_DEFINITION_START_REG(  UBName, Register, Space   ) UB_CB_DEFINITION_START_REG(UBName,Register,Space)
+	#define UB_CBO_MEMBER_NAME(           UBName,         MemberName) UB_CB_MEMBER_NAME(UBName,MemberName)
+	#define UB_CBO_PREFIXED_MEMBER_NAME(  UBName, Prefix, MemberName) UB_CB_PREFIXED_MEMBER_NAME(UBName,Prefix,MemberName)
+	#define UB_CBO_DEFINITION_END(        UBName                    ) UB_CB_DEFINITION_END(UBName)
+	#define UB_CBO_DEFINITION_END_REG(    UBName, Register, Space   ) UB_CB_DEFINITION_END_REG(UBName,Register,Space)
+	#define UB_CBO_MEMBER_ACCESS(         UBName,         MemberName) UB_CB_MEMBER_ACCESS(UBName,MemberName)
+	#define UB_CBO_PREFIXED_MEMBER_ACCESS(UBName, Prefix, MemberName) UB_CB_PREFIXED_MEMBER_ACCESS(UBName,Prefix,MemberName)
 
-	#define UB_DECL_PARAMETER(UBName, StructName, GlobalName) UBName##.##StructName = UBName##_##GlobalName
-	#define UB_DECL_RESOURCE(UBName, StructName, GlobalName)  UBName##.##StructName = UBName##_##GlobalName
-#endif
+	#define UB_CBO_DECL_PARAMETER(UBName, StructName, GlobalName) UB_CB_DECL_PARAMETER(UBName,StructName,GlobalName)
+	#define UB_CBO_DECL_RESOURCE( UBName, StructName, GlobalName) UB_CB_DECL_RESOURCE(UBName,StructName,GlobalName)
 
-#define UB_CB_UNIFORM_BLOCK(UBName, MemberName) UB_FLOAT(4) UB_CB_MEMBER_NAME(UBName,MemberName)[PLATFORM_MAX_UNIFORM_BUFFER_RANGE_FLOAT4]
+	#define UB_CBO_UNIFORM_BLOCK(UBName, MemberName) UB_CB_UNIFORM_BLOCK(UBName, MemberName)
+#endif
 
 // Use interpolator for platforms that do not support ClipDistance
 #if !PLATFORM_SUPPORTS_CLIP_DISTANCE
diff --git a/Engine/Source/Developer/VulkanShaderFormat/Private/VulkanShaderCompiler.cpp b/Engine/Source/Developer/VulkanShaderFormat/Private/VulkanShaderCompiler.cpp
--- a/Engine/Source/Developer/VulkanShaderFormat/Private/VulkanShaderCompiler.cpp
+++ b/Engine/Source/Developer/VulkanShaderFormat/Private/VulkanShaderCompiler.cpp
@@ -223,7 +223,7 @@ struct FVulkanShaderParameterParserPlatformConfiguration : public FShaderParamet
 		// Create a _RootShaderParameters and bind it in slot 0 like any other uniform buffer
 		if (Input.Target.GetFrequency() == SF_RayGen && Input.RootParametersStructure != nullptr)
 		{
-			ConstantBufferType = TEXTVIEW("cbuffer");
+			StableConstantBufferType = TEXTVIEW("cbuffer");
 			EnumAddFlags(Flags, EShaderParameterParserConfigurationFlags::UseStableConstantBuffer);
 		}
 
diff --git a/Engine/Source/Developer/Windows/ShaderFormatD3D/Private/D3DShaderCompiler.cpp b/Engine/Source/Developer/Windows/ShaderFormatD3D/Private/D3DShaderCompiler.cpp
--- a/Engine/Source/Developer/Windows/ShaderFormatD3D/Private/D3DShaderCompiler.cpp
+++ b/Engine/Source/Developer/Windows/ShaderFormatD3D/Private/D3DShaderCompiler.cpp
@@ -1432,6 +1432,17 @@ struct FD3DShaderParameterParserPlatformConfiguration : public FShaderParameterP
 			IndexString.Len(), IndexString.GetData()
 		);
 	}
+
+	virtual FString GetStableConstantBufferRegisterString(EShaderFrequency ShaderFrequency) const override
+	{
+		const uint32 Space = GetAutoBindingSpace(ShaderFrequency);
+		if (Space != 0)
+		{
+			return FString::Printf(TEXT(" : register(b0, space%d)"), Space);
+		}
+
+		return TEXT(" : register(b0)");
+	}
 	
 	const bool bIsRayTracingShader;
 	const FString HitGroupSystemIndexBufferName;
diff --git a/Engine/Source/Developer/Windows/ShaderFormatD3D/Private/D3DShaderCompilerDXC.cpp b/Engine/Source/Developer/Windows/ShaderFormatD3D/Private/D3DShaderCompilerDXC.cpp
--- a/Engine/Source/Developer/Windows/ShaderFormatD3D/Private/D3DShaderCompilerDXC.cpp
+++ b/Engine/Source/Developer/Windows/ShaderFormatD3D/Private/D3DShaderCompilerDXC.cpp
@@ -67,25 +67,6 @@ static bool IsGlobalConstantBufferSupported(const FShaderTarget& Target)
 	}
 }
 
-static uint32 GetAutoBindingSpace(const FShaderTarget& Target)
-{
-	switch (Target.Frequency)
-	{
-	case SF_RayGen:
-		return UE_HLSL_SPACE_RAY_TRACING_GLOBAL;
-	case SF_RayMiss:
-	case SF_RayHitGroup:
-	case SF_RayCallable:
-		return UE_HLSL_SPACE_RAY_TRACING_LOCAL;
-	case SF_WorkGraphRoot:
-		return UE_HLSL_SPACE_WORK_GRAPH_GLOBAL;
-	case SF_WorkGraphComputeNode:
-		return UE_HLSL_SPACE_WORK_GRAPH_LOCAL;
-	default:
-		return UE_HLSL_SPACE_DEFAULT;
-	}
-}
-
 // DXC specific error codes cannot be translated by FPlatformMisc::GetSystemErrorMessage, so do it manually.
 // Codes defines in <DXC>/include/dxc/Support/ErrorCodes.h
 static const TCHAR* DxcErrorCodeToString(HRESULT Code)
diff --git a/Engine/Source/Developer/Windows/ShaderFormatD3D/Private/ShaderFormatD3D.cpp b/Engine/Source/Developer/Windows/ShaderFormatD3D/Private/ShaderFormatD3D.cpp
--- a/Engine/Source/Developer/Windows/ShaderFormatD3D/Private/ShaderFormatD3D.cpp
+++ b/Engine/Source/Developer/Windows/ShaderFormatD3D/Private/ShaderFormatD3D.cpp
@@ -240,8 +240,7 @@ public:
 		Input.Environment.SetDefine(TEXT("COMPILER_DXC"),  bDXC);
 
 		// Do we need SM6.0+ features enabled? This is intentionally disconnected from the ED3DShaderModel to allow SM6.0 to be used without new language features.
-		// TODO: enable PLATFORM_SUPPORTS_CONSTANTBUFFER_OBJECT once we can get usage info from the constant buffer struct.
-		//Input.Environment.SetDefine(TEXT("PLATFORM_SUPPORTS_CONSTANTBUFFER_OBJECT"), bSM6Features);
+		Input.Environment.SetDefine(TEXT("PLATFORM_SUPPORTS_CONSTANTBUFFER_OBJECT"), bSM6Features);
 		Input.Environment.SetDefine(TEXT("PLATFORM_SUPPORTS_SM6_0_WAVE_OPERATIONS"), bSM6Features);
 		Input.Environment.SetDefine(TEXT("PLATFORM_SUPPORTS_DIAGNOSTIC_BUFFER"),     bSM6Features);
 
diff --git a/Engine/Source/Developer/Windows/ShaderFormatD3D/Private/ShaderFormatD3D.h b/Engine/Source/Developer/Windows/ShaderFormatD3D/Private/ShaderFormatD3D.h
--- a/Engine/Source/Developer/Windows/ShaderFormatD3D/Private/ShaderFormatD3D.h
+++ b/Engine/Source/Developer/Windows/ShaderFormatD3D/Private/ShaderFormatD3D.h
@@ -2,6 +2,7 @@
 
 #pragma once
 
+#include "HLSLReservedSpaces.h"
 #include "Windows/WindowsHWrapper.h"
 #include "ShaderCompilerCommon.h"
 
@@ -23,6 +24,30 @@ inline bool DoesShaderModelRequireDXC(ED3DShaderModel ShaderModel)
 	return ShaderModel >= ED3DShaderModel::SM6_0;
 }
 
+inline uint32 GetAutoBindingSpace(EShaderFrequency ShaderFrequency)
+{
+	switch (ShaderFrequency)
+	{
+	case SF_RayGen:
+		return UE_HLSL_SPACE_RAY_TRACING_GLOBAL;
+	case SF_RayMiss:
+	case SF_RayHitGroup:
+	case SF_RayCallable:
+		return UE_HLSL_SPACE_RAY_TRACING_LOCAL;
+	case SF_WorkGraphRoot:
+		return UE_HLSL_SPACE_WORK_GRAPH_GLOBAL;
+	case SF_WorkGraphComputeNode:
+		return UE_HLSL_SPACE_WORK_GRAPH_LOCAL;
+	default:
+		return UE_HLSL_SPACE_DEFAULT;
+	}
+}
+
+inline uint32 GetAutoBindingSpace(const FShaderTarget& Target)
+{
+	return GetAutoBindingSpace(Target.GetFrequency());
+}
+
 bool PreprocessD3DShader(
 	const FShaderCompilerInput& Input,
 	const FShaderCompilerEnvironment& MergedEnvironment,
diff --git a/Engine/Source/Runtime/RenderCore/Private/ShaderParameterParser.cpp b/Engine/Source/Runtime/RenderCore/Private/ShaderParameterParser.cpp
--- a/Engine/Source/Runtime/RenderCore/Private/ShaderParameterParser.cpp
+++ b/Engine/Source/Runtime/RenderCore/Private/ShaderParameterParser.cpp
@@ -1113,6 +1113,7 @@ static const TCHAR* GetConstantSwizzle(uint16 ByteOffset)
 
 bool FShaderParameterParser::MoveShaderParametersToRootConstantBuffer(
 	const FShaderParametersMetadata* RootParametersStructure,
+	EShaderFrequency ShaderFrequency,
 	FString& PreprocessedShaderSource)
 {
 	bool bSuccess = true;
@@ -1122,7 +1123,7 @@ bool FShaderParameterParser::MoveShaderParametersToRootConstantBuffer(
 	{
 		FStringBuilderBase ConstantBufferCode;
 
-		ConstantBufferCode << PlatformConfiguration.ConstantBufferType << TEXT(" ") << FShaderParametersMetadata::kRootUniformBufferBindingName << TEXT("\n{\n");
+		ConstantBufferCode << PlatformConfiguration.StableConstantBufferType << TEXT(" ") << FShaderParametersMetadata::kRootUniformBufferBindingName << PlatformConfiguration.GetStableConstantBufferRegisterString(ShaderFrequency) << TEXT("\n{\n");
 
 		IterateShaderParameterMembers(
 			*RootParametersStructure,
@@ -1236,7 +1237,7 @@ bool FShaderParameterParser::ParseAndModify(const FShaderCompilerInput& Compiler
 
 		if (bNeedToMoveToRootConstantBuffer)
 		{
-			bResult = MoveShaderParametersToRootConstantBuffer(CompilerInput.RootParametersStructure, PreprocessedShaderSource);
+			bResult = MoveShaderParametersToRootConstantBuffer(CompilerInput.RootParametersStructure, CompilerInput.Target.GetFrequency(), PreprocessedShaderSource);
 		}
 
 #if DO_GUARD_SLOW
diff --git a/Engine/Source/Runtime/RenderCore/Private/ShaderParameters.cpp b/Engine/Source/Runtime/RenderCore/Private/ShaderParameters.cpp
--- a/Engine/Source/Runtime/RenderCore/Private/ShaderParameters.cpp
+++ b/Engine/Source/Runtime/RenderCore/Private/ShaderParameters.cpp
@@ -54,10 +54,6 @@ void FShaderResourceParameter::Bind(const FShaderParameterMap& ParameterMap, con
 {
 	if (TOptional<FParameterAllocation> Allocation = ParameterMap.FindParameterAllocation(ParameterName))
 	{
-		if (IsParameterBindless(Allocation->Type))
-		{
-			checkf(Allocation->BufferIndex == 0, TEXT("Unexpected buffer index (%d) for bindless index. Global bindless parameters are expected to be in the global constant buffer (buffer index 0)."), Allocation->BufferIndex);
-		}
 		BaseIndex = Allocation->BaseIndex;
 		NumResources = Allocation->Size;
 		Type = Allocation->Type;
@@ -156,9 +152,9 @@ static void CreateHLSLUniformBufferStructMembersDeclaration(
 
 	auto AddStructMember = [&](const FShaderParametersMetadata::FMember& Member, bool bResource)
 	{
-		// UB_DECL_PARAMETER(UniformBufferName, StructPrefix_Name, GlobalPrefix_Name)
+		// UB_CB_DECL_PARAMETER(UniformBufferName, StructPrefix_Name, GlobalPrefix_Name)
 
-		Decl.StructMembers << (bResource ? TEXT("UB_DECL_RESOURCE(") : TEXT("UB_DECL_PARAMETER(")) << UniformBufferName << TEXT(",");
+		Decl.StructMembers << (bResource ? TEXT("UB_CB_DECL_RESOURCE(") : TEXT("UB_CB_DECL_PARAMETER(")) << UniformBufferName << TEXT(",");
 		if (!StructPrefix.IsEmpty())
 		{
 			Decl.StructMembers << StructPrefix << TEXT(".");
@@ -385,48 +381,37 @@ static FString CreateHLSLUniformBufferDeclaration(const TCHAR* UniformBufferName
 		uint32 HLSLBaseOffset = 0;
 		CreateHLSLUniformBufferStructMembersDeclaration(UniformBufferStruct, UniformBufferName, TEXT(""), TEXT(""), 0, Decl, HLSLBaseOffset);
 
-		if (UniformBufferSBLayout && UniformBufferSBLayout->RegisterSpace > 0)
-		{	
-			return FString::Printf(
-				TEXT("#pragma once\n")
-				TEXT("UB_STATIC_CB_DEFINITION_START(%s,%d,%d)\n")
-				TEXT("%s")
-				TEXT("UB_CB_DEFINITION_END(%s)\n")
-				TEXT("%s")
-				TEXT("UniformBuffer %s\n")
-				TEXT("{\n")
-				TEXT("%s")
-				TEXT("};\n"),
-				UniformBufferName,
-				UniformBufferSBLayout->CBVResourceIndex,
-				UniformBufferSBLayout->RegisterSpace,
-				*Decl.ConstantBufferMembers,
-				UniformBufferName,
-				*Decl.ResourceMembers,
-				UniformBufferName,
-				*Decl.StructMembers
-			);
-		}
-		else
+		const bool bUseRegister = (UniformBufferSBLayout && UniformBufferSBLayout->RegisterSpace > 0);
+
+		// UB_CB_DEFINITION_START and UB_CB_DEFINITION_END both have variants that take the register + space specifiers, so use those when appropriate
+		const FString StartEndSuffix = bUseRegister
+			? FString::Printf(TEXT("_REG(%s,%d,%d)"), UniformBufferName, UniformBufferSBLayout->CBVResourceIndex, UniformBufferSBLayout->RegisterSpace)
+			: FString::Printf(TEXT("(%s)"), UniformBufferName);
+
+		FString Result = FString::Printf(
+			TEXT("#pragma once\n")
+			TEXT("UB_CB_DEFINITION_START%s\n")
+			TEXT("%s")
+			TEXT("UB_CB_DEFINITION_END%s\n")
+			TEXT("%s")
+			TEXT("UniformBuffer %s\n")
+			TEXT("{\n")
+			TEXT("%s")
+			TEXT("};\n"),
+			*StartEndSuffix,
+			*Decl.ConstantBufferMembers,
+			*StartEndSuffix,
+			*Decl.ResourceMembers,
+			UniformBufferName,
+			*Decl.StructMembers
+		);
+
+		if (Decl.ResourceMembers.Len() == 0)
 		{
-			return FString::Printf(
-				TEXT("#pragma once\n")
-				TEXT("UB_CB_DEFINITION_START(%s)\n")
-				TEXT("%s")
-				TEXT("UB_CB_DEFINITION_END(%s)\n")
-				TEXT("%s")
-				TEXT("UniformBuffer %s\n")
-				TEXT("{\n")
-				TEXT("%s")
-				TEXT("};\n"),
-				UniformBufferName,
-				*Decl.ConstantBufferMembers,
-				UniformBufferName,
-				*Decl.ResourceMembers,
-				UniformBufferName,
-				*Decl.StructMembers
-			);
+			return Result.Replace(TEXT("UB_CB_"), TEXT("UB_CBO_"), ESearchCase::CaseSensitive);
 		}
+
+		return Result;
 	}
 
 	return FString(TEXT("\n"));
diff --git a/Engine/Source/Runtime/RenderCore/Public/ShaderCore.h b/Engine/Source/Runtime/RenderCore/Public/ShaderCore.h
--- a/Engine/Source/Runtime/RenderCore/Public/ShaderCore.h
+++ b/Engine/Source/Runtime/RenderCore/Public/ShaderCore.h
@@ -752,6 +752,7 @@ enum class EShaderOptionalDataKey : uint8
 	VendorExtension      = uint8('v'),
 	ShaderBindingLayout  = uint8('s'),
 	EntryPoint           = uint8('e'),
+	GlobalUniformBufferConfiguration = uint8('G'),
 };
 
 enum class EShaderResourceUsageFlags : uint8
@@ -780,6 +781,15 @@ struct FShaderCodePackedResourceCounts
 	uint8 NumUAVs;
 };
 
+// Configuration for a different 'global constant buffer' register assignment.
+// Most compilers always force it to 0 or allow you to assign it, but some compilers do not.
+struct FShaderCodeGlobalUniformBufferConfiguration
+{
+	static const EShaderOptionalDataKey Key = EShaderOptionalDataKey::GlobalUniformBufferConfiguration;
+
+	uint32 Slot;
+};
+
 struct FShaderCodeResourceMasks
 {
 	// for FindOptionalData() and AddOptionalData()
diff --git a/Engine/Source/Runtime/RenderCore/Public/ShaderParameterParser.h b/Engine/Source/Runtime/RenderCore/Public/ShaderParameterParser.h
--- a/Engine/Source/Runtime/RenderCore/Public/ShaderParameterParser.h
+++ b/Engine/Source/Runtime/RenderCore/Public/ShaderParameterParser.h
@@ -43,11 +43,11 @@ public:
 	struct FPlatformConfiguration
 	{
 		FPlatformConfiguration() = default;
-		FPlatformConfiguration(FStringView InConstantBufferType, EShaderParameterParserConfigurationFlags InFlags = EShaderParameterParserConfigurationFlags::None)
-			: ConstantBufferType(InConstantBufferType)
+		FPlatformConfiguration(FStringView InStableConstantBufferType, EShaderParameterParserConfigurationFlags InFlags = EShaderParameterParserConfigurationFlags::None)
+			: StableConstantBufferType(InStableConstantBufferType)
 			, Flags(InFlags)
 		{
-			if (InConstantBufferType.Len())
+			if (InStableConstantBufferType.Len())
 			{
 				EnumAddFlags(Flags, EShaderParameterParserConfigurationFlags::UseStableConstantBuffer);
 			}
@@ -70,9 +70,16 @@ public:
 			return FString();
 		}
 
+		/** Since the register/space for the global constant buffer can vary based on frequency, it's up to the shader format to communicate that to us */
+		virtual FString GetStableConstantBufferRegisterString(EShaderFrequency ShaderFrequency) const
+		{
+			// Not required to be overridden
+			return FString();
+		}
+
 		TConstArrayView<FStringView> ExtraSRVTypes;
 		TConstArrayView<FStringView> ExtraUAVTypes;
-		FStringView ConstantBufferType;
+		FStringView StableConstantBufferType;
 		EShaderParameterParserConfigurationFlags Flags = EShaderParameterParserConfigurationFlags::None;
 	};
 
@@ -217,6 +224,7 @@ protected:
 	/** Moves parsed parameters into the root constant buffer. */
 	RENDERCORE_API bool MoveShaderParametersToRootConstantBuffer(
 		const FShaderParametersMetadata* RootParametersStructure,
+		EShaderFrequency ShaderFrequency,
 		FString& PreprocessedShaderSource
 	);
 
