Simplify logic to apply Relaxed Precision with PLATFORM_SUPPORTS_RELAXED_PRECISION for supported platforms
Also restrict Relaxed Precision for MOBILE_EMULATION to platforms that actually use it on their mobile counterpart (ie: METAL_ES3_1)

#tests Preflight Android & iOS tested by [REDACTED_PERSON]

diff --git a/Engine/Shaders/Public/Platform.ush b/Engine/Shaders/Public/Platform.ush
--- a/Engine/Shaders/Public/Platform.ush
+++ b/Engine/Shaders/Public/Platform.ush
@@ -252,6 +252,10 @@
 #define PLATFORM_SUPPORTS_REAL_TYPES 0
 #endif
 
+#ifndef PLATFORM_SUPPORTS_RELAXED_PRECISION
+#define PLATFORM_SUPPORTS_RELAXED_PRECISION 0
+#endif
+
 #ifndef COMPILER_SUPPORTS_WAVE_32_64_MODE
 #define COMPILER_SUPPORTS_WAVE_32_64_MODE 0
 #endif
@@ -335,7 +339,7 @@
 	#pragma dxc diagnostic ignored "-Wconversion"
 #endif
 
-#if ES3_1_PROFILE && MOBILE_EMULATION && !FORCE_FLOATS
+#if !FORCE_FLOATS && PLATFORM_SUPPORTS_RELAXED_PRECISION
 	#define half min16float
 	#define half1 min16float1
 	#define half2 min16float2
@@ -344,9 +348,7 @@
 	#define half3x3 min16float3x3
 	#define half3x4 min16float3x4
 	#define half4x4 min16float4x4
-#elif ((!(ES3_1_PROFILE || METAL_ES3_1_PROFILE || VULKAN_PROFILE_SM5)) && !PLATFORM_SUPPORTS_REAL_TYPES) || FORCE_FLOATS
-	// Always use floats when using the ES3/METAL compiler, because platforms not optimized for lower precision, 
-	// And we don't want potential side effects on other platforms
+#elif FORCE_FLOATS || !PLATFORM_SUPPORTS_REAL_TYPES
 	#define half float
 	#define half1 float1
 	#define half2 float2
@@ -363,15 +365,6 @@
 	#define fixed3x3 float3x3
 	#define fixed4x4 float4x4
 	#define fixed4x3 float4x3
-#elif (VULKAN_PROFILE || VULKAN_PROFILE_SM5) || (COMPILER_GLSL_ES3_1 && !(COMPILER_HLSLCC && COMPILER_HLSLCC == 1))
-	// For VULKAN and OPENGL ES31 use RelaxedPrecision for half floats
-	#define half min16float
-	#define half2 min16float2
-	#define half3 min16float3
-	#define half4 min16float4
-	#define half3x3 min16float3x3
-	#define half3x4 min16float3x4 
-	#define half4x4 min16float4x4
 #endif
 
 // ---------------------------------------------------- Profile config
diff --git a/Engine/Source/Developer/ShaderFormatOpenGL/Private/ShaderFormatOpenGL.cpp b/Engine/Source/Developer/ShaderFormatOpenGL/Private/ShaderFormatOpenGL.cpp
--- a/Engine/Source/Developer/ShaderFormatOpenGL/Private/ShaderFormatOpenGL.cpp
+++ b/Engine/Source/Developer/ShaderFormatOpenGL/Private/ShaderFormatOpenGL.cpp
@@ -83,6 +83,7 @@ public:
 		case GLSL_ES3_1_ANDROID:
 			Input.Environment.SetDefine(TEXT("COMPILER_GLSL_ES3_1"), 1);
 			Input.Environment.SetDefine(TEXT("ES3_1_PROFILE"), 1);
+			Input.Environment.SetDefine(TEXT("PLATFORM_SUPPORTS_RELAXED_PRECISION"), 1);
 			break;
 
 		case GLSL_150_ES3_1:
diff --git a/Engine/Source/Developer/VulkanShaderFormat/Private/SpirVShaderCompiler.inl b/Engine/Source/Developer/VulkanShaderFormat/Private/SpirVShaderCompiler.inl
--- a/Engine/Source/Developer/VulkanShaderFormat/Private/SpirVShaderCompiler.inl
+++ b/Engine/Source/Developer/VulkanShaderFormat/Private/SpirVShaderCompiler.inl
@@ -1396,6 +1396,7 @@ static void ModifyCompilerInput(FSpirvShaderCompilerInternalState& InternalState
 	{
 		Input.Environment.SetDefine(TEXT("ES3_1_PROFILE"), 1);
 		Input.Environment.SetDefine(TEXT("VULKAN_PROFILE"), 1);
+		Input.Environment.SetDefine(TEXT("PLATFORM_SUPPORTS_RELAXED_PRECISION"), 1);
 	}
 	else if (InternalState.IsSM6())
 	{
@@ -1405,6 +1406,7 @@ static void ModifyCompilerInput(FSpirvShaderCompilerInternalState& InternalState
 	else if (InternalState.IsSM5())
 	{
 		Input.Environment.SetDefine(TEXT("VULKAN_PROFILE_SM5"), 1);
+		Input.Environment.SetDefine(TEXT("PLATFORM_SUPPORTS_RELAXED_PRECISION"), 1);
 	}
 	Input.Environment.SetDefine(TEXT("row_major"), TEXT(""));
 
