Vulkan static samplers now start at heap offset 0.

diff --git a/Engine/Shaders/Public/Platform/Vulkan/VulkanCommon.ush b/Engine/Shaders/Public/Platform/Vulkan/VulkanCommon.ush
--- a/Engine/Shaders/Public/Platform/Vulkan/VulkanCommon.ush
+++ b/Engine/Shaders/Public/Platform/Vulkan/VulkanCommon.ush
@@ -141,12 +141,12 @@ void ImageInterlockedOrUInt64(RWTexture2D<UlongType> Dest, uint2 Coord, UlongTyp
 #if UE_BINDLESS_ENABLED
 	typedef SamplerState SafeTypeGlobalSamplers;
 	SafeTypeGlobalSamplers VULKAN_SAMPLER_HEAP(GlobalSamplers)[];
-	static const SamplerState VulkanGlobalPointClampedSampler     = VULKAN_SAMPLER_HEAP(GlobalSamplers)[1];
-	static const SamplerState VulkanGlobalPointWrappedSampler     = VULKAN_SAMPLER_HEAP(GlobalSamplers)[2];
-	static const SamplerState VulkanGlobalBilinearClampedSampler  = VULKAN_SAMPLER_HEAP(GlobalSamplers)[3];
-	static const SamplerState VulkanGlobalBilinearWrappedSampler  = VULKAN_SAMPLER_HEAP(GlobalSamplers)[4];
-	static const SamplerState VulkanGlobalTrilinearClampedSampler = VULKAN_SAMPLER_HEAP(GlobalSamplers)[5];
-	static const SamplerState VulkanGlobalTrilinearWrappedSampler = VULKAN_SAMPLER_HEAP(GlobalSamplers)[6];
+	static const SamplerState VulkanGlobalPointClampedSampler     = VULKAN_SAMPLER_HEAP(GlobalSamplers)[0];
+	static const SamplerState VulkanGlobalPointWrappedSampler     = VULKAN_SAMPLER_HEAP(GlobalSamplers)[1];
+	static const SamplerState VulkanGlobalBilinearClampedSampler  = VULKAN_SAMPLER_HEAP(GlobalSamplers)[2];
+	static const SamplerState VulkanGlobalBilinearWrappedSampler  = VULKAN_SAMPLER_HEAP(GlobalSamplers)[3];
+	static const SamplerState VulkanGlobalTrilinearClampedSampler = VULKAN_SAMPLER_HEAP(GlobalSamplers)[4];
+	static const SamplerState VulkanGlobalTrilinearWrappedSampler = VULKAN_SAMPLER_HEAP(GlobalSamplers)[5];
 #else
 	SamplerState VulkanGlobalPointWrappedSampler;
 	SamplerState VulkanGlobalPointClampedSampler;
