Metal RHI

Moved Uniform Buffers to be the first indices assigned in Metal Spirv backend.

diff --git a/Engine/Shaders/Public/Platform/Metal/MetalCommon.ush b/Engine/Shaders/Public/Platform/Metal/MetalCommon.ush
--- a/Engine/Shaders/Public/Platform/Metal/MetalCommon.ush
+++ b/Engine/Shaders/Public/Platform/Metal/MetalCommon.ush
@@ -7,7 +7,7 @@
 #pragma once
 
 // Update this GUID to invalidate shader recompilation for only Metal shaders
-#pragma message("UESHADERMETADATA_VERSION 80E79B41-0792-4136-8D8E-BB1BA48C008D")
+#pragma message("UESHADERMETADATA_VERSION 80E79B41-0792-4136-8D8E-BB1BA48C008E")
 
 #if VERTEXSHADER && HULLSHADER
 #ifndef TESSELLATIONSHADER
