Add Format=0 support to SubstratePublic.ush

diff --git a/Engine/Shaders/Private/MaterialTemplate.ush b/Engine/Shaders/Private/MaterialTemplate.ush
--- a/Engine/Shaders/Private/MaterialTemplate.ush
+++ b/Engine/Shaders/Private/MaterialTemplate.ush
@@ -3300,7 +3300,7 @@ float4 SceneTextureLookup(float2 UV, int SceneTextureIndex, bool bFiltered)
 		case PPI_IsFirstPerson:			return 0.0f;
 		}
 	}
-	#elif SUBSTRATE_GBUFFER_FORMAT==1 && !SHADING_PATH_MOBILE // Not on mobile that is not using dbuffer decal but the legacy GBuffer blending.
+	#elif SUBSTRATE_ENABLED && (SUBSTRATE_GBUFFER_FORMAT==0 || SUBSTRATE_GBUFFER_FORMAT==1) && !SHADING_PATH_MOBILE // Not on mobile that is not using dbuffer decal but the legacy GBuffer blending.
 	{
 		#if !SUBSTRATE_OPAQUE_DEFERRED || ALLOW_SUBSTRATE_PUBLIC_ACCESS // Prevents opaque material to sample the Gbuffer.
 		FSubstrateGBuffer SubstrateGBuffer = SubstratePublic_GetSubstrateGBufferFromFirstBSDF(PixelPos);
diff --git a/Engine/Shaders/Public/SubstratePublic.ush b/Engine/Shaders/Public/SubstratePublic.ush
--- a/Engine/Shaders/Public/SubstratePublic.ush
+++ b/Engine/Shaders/Public/SubstratePublic.ush
@@ -7,10 +7,11 @@
 =============================================================================*/
 
 #include "/Engine/Private/Substrate/Substrate.ush"
+#include "/Engine/Private/Substrate/SubstrateRead.ush"
 
 #if SUBSTRATE_ENABLED 
 
-#define SUBSTRATE_ALLOWS_GBUFFER_SAMPLING  ((NEEDS_SCENE_TEXTURES && !SCENE_TEXTURES_DISABLED) && SUBSTRATE_MATERIALCONTAINER_IS_VIEWRESOURCE)
+#define SUBSTRATE_ALLOWS_GBUFFER_SAMPLING  ((NEEDS_SCENE_TEXTURES && !SCENE_TEXTURES_DISABLED))
 
 /** Public facing functions allowing external users access to some Substrate data. */
 
@@ -43,11 +44,12 @@ struct FSubstrateGBuffer
 
 FSubstrateGBuffer SubstratePublic_GetSubstrateGBufferFromFirstBSDF(uint2 PixelCoord)
 {
-	FSubstrateGBuffer SGBuffer = (FSubstrateGBuffer)0; 
-	SGBuffer.bIsValid = false;
-	SGBuffer.WorldNormal = float3(0.0, 0.0, 1.0);
+	FSubstrateGBuffer Out = (FSubstrateGBuffer)0; 
+	Out.bIsValid = false;
+	Out.WorldNormal = float3(0.0, 0.0, 1.0);
 
 #if SUBSTRATE_ALLOWS_GBUFFER_SAMPLING
+	#if SUBSTRATE_GBUFFER_FORMAT == 1 && SUBSTRATE_MATERIALCONTAINER_IS_VIEWRESOURCE
 	FSubstrateAddressing SubstrateAddressing   = GetSubstratePixelDataByteOffset(PixelCoord, uint2(View.BufferSizeAndInvSize.xy), SubstratePublicStruct.MaxBytesPerPixel);
 	FSubstratePixelHeader SubstratePixelHeader = UnpackSubstrateHeaderIn(SubstratePublicStruct.MaterialTextureArray, SubstrateAddressing, SubstratePublicStruct.TopLayerTexture);
 	
@@ -59,34 +61,34 @@ FSubstrateGBuffer SubstratePublic_GetSubstrateGBufferFromFirstBSDF(uint2 PixelCo
 		const bool bSkipSSSMaterialOverride = true; // We do not want the material data to be affected by any override.
 		FSubstrateBSDF BSDF = UnpackSubstrateBSDFIn(SubstratePublicStruct.MaterialTextureArray, SubstrateAddressing, SubstratePixelHeader, bSkipSSSMaterialOverride);
 	
-		SGBuffer.bIsValid			= true;
-		SGBuffer.DiffuseColor		= SubstrateGetBSDFDiffuseColor(BSDF);
-		SGBuffer.SpecularColor		= SubstrateGetBSDFSpecularF0(BSDF);
-		SGBuffer.SubsurfaceColor	= SubstrateGetBSDFSubSurfaceColor(BSDF);
-		SGBuffer.BaseColor			= SubstrateGetBSDFBaseColor(BSDF);
-		SGBuffer.Specular			= SubstrateGetBSDFSpecular(BSDF);
-		SGBuffer.Metallic			= SubstrateGetBSDFMetallic(BSDF);
-		SGBuffer.WorldNormal		= SubstrateGetWorldNormal(SubstratePixelHeader, BSDF, SubstrateAddressing);
-		SGBuffer.Roughness			= SubstrateGetBSDFRoughness(BSDF);
-		SGBuffer.MaterialAO			= SubstrateGetAO(SubstratePixelHeader);
-		SGBuffer.ShadingModelID		= SubstrateGetLegacyShadingModels(BSDF);
-		SGBuffer.ShadingModelColor	= GetShadingModelColor(SGBuffer.ShadingModelID);
-		SGBuffer.StoredBaseColor	= SubstrateGetBSDFBaseColor(BSDF);
-		SGBuffer.StoredSpecular		= SubstrateGetBSDFSpecular(BSDF).rrr;
-		SGBuffer.WorldTangent		= SubstrateGetWorldTangent(SubstratePixelHeader, BSDF, SubstrateAddressing);
-		SGBuffer.Anisotropy			= SubstrateGetBSDFAnisotropy(BSDF);
-		SGBuffer.bIsFirstPerson		= SubstratePixelHeader.IsFirstPerson();
-
-		SGBuffer.Opacity			= 0.0f;
+		Out.bIsValid			= true;
+		Out.DiffuseColor		= SubstrateGetBSDFDiffuseColor(BSDF);
+		Out.SpecularColor		= SubstrateGetBSDFSpecularF0(BSDF);
+		Out.SubsurfaceColor		= SubstrateGetBSDFSubSurfaceColor(BSDF);
+		Out.BaseColor			= SubstrateGetBSDFBaseColor(BSDF);
+		Out.Specular			= SubstrateGetBSDFSpecular(BSDF);
+		Out.Metallic			= SubstrateGetBSDFMetallic(BSDF);
+		Out.WorldNormal			= SubstrateGetWorldNormal(SubstratePixelHeader, BSDF, SubstrateAddressing);
+		Out.Roughness			= SubstrateGetBSDFRoughness(BSDF);
+		Out.MaterialAO			= SubstrateGetAO(SubstratePixelHeader);
+		Out.ShadingModelID		= SubstrateGetLegacyShadingModels(BSDF);
+		Out.ShadingModelColor	= GetShadingModelColor(Out.ShadingModelID);
+		Out.StoredBaseColor		= SubstrateGetBSDFBaseColor(BSDF);
+		Out.StoredSpecular		= SubstrateGetBSDFSpecular(BSDF).rrr;
+		Out.WorldTangent		= SubstrateGetWorldTangent(SubstratePixelHeader, BSDF, SubstrateAddressing);
+		Out.Anisotropy			= SubstrateGetBSDFAnisotropy(BSDF);
+		Out.bIsFirstPerson		= SubstratePixelHeader.IsFirstPerson();
+
+		Out.Opacity			= 0.0f;
 		if (SubstrateSubSurfaceHeaderGetSSSType(SSSHeader) != SSS_TYPE_NONE)
 		{
-			SGBuffer.Opacity = SubstrateSubSurfaceHeaderGetOpacity(SSSHeader);
+			Out.Opacity = SubstrateSubSurfaceHeaderGetOpacity(SSSHeader);
 		}
 		else if (BSDF_GETTYPE(BSDF) == SUBSTRATE_BSDF_TYPE_SLAB)
 		{
 			if (BSDF_GETHASFUZZ(BSDF))
 			{
-				SGBuffer.Opacity = SLAB_FUZZ_AMOUNT(BSDF);	// Cloth
+				Out.Opacity = SLAB_FUZZ_AMOUNT(BSDF);	// Cloth
 			}
 			else if (BSDF_GETHASHAZINESS(BSDF))
 			{
@@ -94,46 +96,75 @@ FSubstrateGBuffer SubstratePublic_GetSubstrateGBufferFromFirstBSDF(uint2 PixelCo
 				const bool bHazeAsSimpleClearCoat = Haziness.bSimpleClearCoat;
 				if (bHazeAsSimpleClearCoat)
 				{
-					SGBuffer.Opacity = Haziness.Weight;		// ClearCoat
+					Out.Opacity = Haziness.Weight;		// ClearCoat
 				}
 			}
 		}
 		else if (BSDF_GETTYPE(BSDF) == SUBSTRATE_BSDF_TYPE_EYE)
 		{
-			SGBuffer.Opacity = 1 - EYE_IRISMASK(BSDF);		// IrisMask
+			Out.Opacity = 1 - EYE_IRISMASK(BSDF);		// IrisMask
 		}
 		else if (BSDF_GETTYPE(BSDF) == SUBSTRATE_BSDF_TYPE_HAIR)
 		{
-			SGBuffer.Opacity = HAIR_BACKLIT(BSDF);			// BackLit
+			Out.Opacity = HAIR_BACKLIT(BSDF);			// BackLit
 		}
 	}
 	else if (SubstratePixelHeader.State == 0)				// Only Unlit and non written pixels have a state of 0. Those pixels should show up as unlit to respect the legacy behavior.
 	{		
-		SGBuffer.bIsValid			= true;
-		SGBuffer.ShadingModelID		= SHADINGMODELID_UNLIT;
-		SGBuffer.ShadingModelColor	= GetShadingModelColor(SGBuffer.ShadingModelID);
+		Out.bIsValid			= true;
+		Out.ShadingModelID		= SHADINGMODELID_UNLIT;
+		Out.ShadingModelColor	= GetShadingModelColor(Out.ShadingModelID);
 	}
+	#elif SUBSTRATE_GBUFFER_FORMAT == 0
+	const float2 UV = (float2(PixelCoord) + 0.5) * View.BufferSizeAndInvSize.zw;
+	const FGBufferData GBuffer = GetScreenSpaceData(UV).GBuffer;
+	Out.bIsValid			= GBuffer.ShadingModelID != SHADINGMODELID_UNLIT;
+	Out.DiffuseColor		= GBuffer.DiffuseColor;
+	Out.SpecularColor		= GBuffer.SpecularColor;
+	Out.SubsurfaceColor		= IsSubsurfaceModel(GBuffer.ShadingModelID) ? ExtractSubsurfaceColor(GBuffer) : GBuffer.CustomData.xyz;
+	Out.BaseColor			= GBuffer.BaseColor;
+	Out.Specular			= GBuffer.Specular;
+	Out.Metallic			= GBuffer.Metallic;
+	Out.WorldNormal			= GBuffer.WorldNormal;
+	Out.Opacity				= HasCustomGBufferData(GBuffer.ShadingModelID) ?  GBuffer.CustomData.a : 1;
+	Out.Roughness			= GBuffer.Roughness;
+	Out.MaterialAO			= GBuffer.GBufferAO;
+	Out.ShadingModelColor	= GetShadingModelColor(GBuffer.ShadingModelID);
+	Out.ShadingModelID		= GBuffer.ShadingModelID;
+	Out.StoredBaseColor		= GBuffer.StoredBaseColor;
+	Out.StoredSpecular		= GBuffer.StoredSpecular;
+	Out.WorldTangent		= GBuffer.WorldTangent;
+	Out.Anisotropy			= GBuffer.Anisotropy;
+	Out.bIsFirstPerson		= IsFirstPerson(GBuffer);
+	#endif
 #endif
-
-	return SGBuffer;
+	return Out;
 }
 
 // Optimised data fetch from the TopLayer texture
 float3 SubstratePublic_GetWorldNormal(uint2 PixelCoord)
 {
+	float3 Out = float3(0, 0, 1);
 #if SUBSTRATE_ALLOWS_GBUFFER_SAMPLING
-  return SubstrateUnpackTopLayerData(SubstratePublicStruct.TopLayerTexture.Load(uint3(PixelCoord, 0))).WorldNormal;
-#else
-	return float3(0, 0, 1);
+	#if SUBSTRATE_GBUFFER_FORMAT == 1 && SUBSTRATE_MATERIALCONTAINER_IS_VIEWRESOURCE
+	Out = SubstrateUnpackTopLayerData(SubstratePublicStruct.TopLayerTexture.Load(uint3(PixelCoord, 0))).WorldNormal;
+	#elif SUBSTRATE_GBUFFER_FORMAT == 0
+	Out = GetScreenSpaceData(UV).GBuffer.WorldNormal;
+	#endif
 #endif
+	return Out;
 }
 float SubstratePublic_GetRoughness(uint2 PixelCoord)
 {
+	float Out = 0;
 #if SUBSTRATE_ALLOWS_GBUFFER_SAMPLING
-  return SubstrateUnpackTopLayerData(SubstratePublicStruct.TopLayerTexture.Load(uint3(PixelCoord, 0))).Roughness;
-#else
-	return 0.0;
+	#if SUBSTRATE_GBUFFER_FORMAT == 1 && SUBSTRATE_MATERIALCONTAINER_IS_VIEWRESOURCE
+	Out = SubstrateUnpackTopLayerData(SubstratePublicStruct.TopLayerTexture.Load(uint3(PixelCoord, 0))).Roughness;
+	#elif SUBSTRATE_GBUFFER_FORMAT == 0
+	Out = GetScreenSpaceData(UV).GBuffer.Roughness;
+	#endif
 #endif
+	return Out;
 }
 
 #endif
\ No newline at end of file
