Substrate - fix compilation issues with SubstratePublic.ush.

#rnx
#tests PC editor dx12

diff --git a/Engine/Shaders/Public/SubstratePublic.ush b/Engine/Shaders/Public/SubstratePublic.ush
--- a/Engine/Shaders/Public/SubstratePublic.ush
+++ b/Engine/Shaders/Public/SubstratePublic.ush
@@ -149,6 +149,7 @@ float3 SubstratePublic_GetWorldNormal(uint2 PixelCoord)
 	#if SUBSTRATE_GBUFFER_FORMAT == 1 && SUBSTRATE_MATERIALCONTAINER_IS_VIEWRESOURCE
 	Out = SubstrateUnpackTopLayerData(SubstratePublicStruct.TopLayerTexture.Load(uint3(PixelCoord, 0))).WorldNormal;
 	#elif SUBSTRATE_GBUFFER_FORMAT == 0
+	const float2 UV = (float2(PixelCoord) + 0.5) * View.BufferSizeAndInvSize.zw;
 	Out = GetScreenSpaceData(UV).GBuffer.WorldNormal;
 	#endif
 #endif
@@ -161,6 +162,7 @@ float SubstratePublic_GetRoughness(uint2 PixelCoord)
 	#if SUBSTRATE_GBUFFER_FORMAT == 1 && SUBSTRATE_MATERIALCONTAINER_IS_VIEWRESOURCE
 	Out = SubstrateUnpackTopLayerData(SubstratePublicStruct.TopLayerTexture.Load(uint3(PixelCoord, 0))).Roughness;
 	#elif SUBSTRATE_GBUFFER_FORMAT == 0
+	const float2 UV = (float2(PixelCoord) + 0.5) * View.BufferSizeAndInvSize.zw;
 	Out = GetScreenSpaceData(UV).GBuffer.Roughness;
 	#endif
 #endif
