Fix Opacity buffer view mode in Format=0.

The opacity value now match the legacy behavior.

diff --git a/Engine/Shaders/Public/SubstratePublic.ush b/Engine/Shaders/Public/SubstratePublic.ush
--- a/Engine/Shaders/Public/SubstratePublic.ush
+++ b/Engine/Shaders/Public/SubstratePublic.ush
@@ -126,7 +126,7 @@ FSubstrateGBuffer SubstratePublic_GetSubstrateGBufferFromFirstBSDF(uint2 PixelCo
 	Out.Specular			= GBuffer.Specular;
 	Out.Metallic			= GBuffer.Metallic;
 	Out.WorldNormal			= GBuffer.WorldNormal;
-	Out.Opacity				= HasCustomGBufferData(GBuffer.ShadingModelID) ?  GBuffer.CustomData.a : 1;
+	Out.Opacity				= GBuffer.CustomData.a;
 	Out.Roughness			= GBuffer.Roughness;
 	Out.MaterialAO			= GBuffer.GBufferAO;
 	Out.ShadingModelColor	= GetShadingModelColor(GBuffer.ShadingModelID);
