* Unify Substrate legacy Opacity fetching.
* Add Opacity to SubstrateExport output as opacity has different meaning than coverage for legacy material.

diff --git a/Engine/Shaders/Private/Substrate/Substrate.ush b/Engine/Shaders/Private/Substrate/Substrate.ush
--- a/Engine/Shaders/Private/Substrate/Substrate.ush
+++ b/Engine/Shaders/Private/Substrate/Substrate.ush
@@ -3782,6 +3782,47 @@ float SubstrateGetBSDFAnisotropy(in FSubstrateBSDF BSDF)
 	return 0.f;
 }
 
+float SubstrateGetBSDFOpacity(in FSubstrateBSDF BSDF)
+{
+	// Recompute the 'legacy' opacity value based on feature/BSDF type
+	float Out = 0.f;
+	if (BSDF_GETTYPE(BSDF) == SUBSTRATE_BSDF_TYPE_SLAB)
+	{
+		if (BSDF_GETSSSTYPE(BSDF) != SSS_TYPE_NONE)
+		{
+			switch (BSDF_GETSSSTYPE(BSDF))
+			{
+				case SSS_TYPE_WRAP				: Out = 1.f-SLAB_SSSPHASEANISOTROPY(BSDF); break;
+				case SSS_TYPE_TWO_SIDED_WRAP	: Out = 1.f-SLAB_SSSPHASEANISOTROPY(BSDF); break;
+				case SSS_TYPE_DIFFUSION			: Out = 1.f; break;
+				case SSS_TYPE_DIFFUSION_PROFILE	: Out = SLAB_SSSPROFILERADIUSSCALE(BSDF); break;
+				case SSS_TYPE_SIMPLEVOLUME		: Out = 1.f; 
+			}
+		}
+		else if (BSDF_GETHASFUZZ(BSDF))
+		{
+			Out = SLAB_FUZZ_AMOUNT(BSDF);	// Cloth
+		}
+		else if (BSDF_GETHASHAZINESS(BSDF))
+		{
+			const FHaziness Haziness = UnpackHaziness(SLAB_HAZINESS(BSDF));
+			if (Haziness.bSimpleClearCoat)
+			{
+				Out = Haziness.Weight;		// ClearCoat
+			}
+		}
+	}
+	else if (BSDF_GETTYPE(BSDF) == SUBSTRATE_BSDF_TYPE_EYE)
+	{
+		Out = 1 - EYE_IRISMASK(BSDF);		// IrisMask
+	}
+	else if (BSDF_GETTYPE(BSDF) == SUBSTRATE_BSDF_TYPE_HAIR)
+	{
+		Out = HAIR_BACKLIT(BSDF);			// BackLit
+	}
+	return Out;
+}
+
 float3 SubstrateGetBSDFEmissive(in FSubstrateBSDF BSDF)
 {	
 	return BSDF_GETEMISSIVE(BSDF);
diff --git a/Engine/Shaders/Private/Substrate/SubstrateExport.ush b/Engine/Shaders/Private/Substrate/SubstrateExport.ush
--- a/Engine/Shaders/Private/Substrate/SubstrateExport.ush
+++ b/Engine/Shaders/Private/Substrate/SubstrateExport.ush
@@ -945,6 +945,7 @@ struct FExportResult
 	float  ShadingModelID;
 	float IndirectIrradianceScale;
 	float Curvature;
+	float Opacity;
 };
 
 // This is used to export very simple data for lightmaps, material baking or preview on editor node.
@@ -1006,7 +1007,7 @@ FExportResult SubstrateMaterialExportOut(
 			Out.Roughness				= SubstrateGetBSDFRoughness(CurrentBSDF);
 			Out.Anisotropy				= SubstrateGetBSDFAnisotropy(CurrentBSDF);
 			Out.Metallic				= SubstrateGetBSDFMetallic(CurrentBSDF);
-
+			Out.Opacity 				= SubstrateGetBSDFOpacity(CurrentBSDF);
 			Out.SubsurfaceColor			= SubstrateGetBSDFSubSurfaceColor(CurrentBSDF);
 			Out.SubsurfaceProfileIDFloat= ConvertSubsurfaceProfileToFloat(SubstrateGetBSDFSubSurfaceProfileID(CurrentBSDF));
 			Out.WorldTangent			= T;
diff --git a/Engine/Shaders/Public/SubstratePublic.ush b/Engine/Shaders/Public/SubstratePublic.ush
--- a/Engine/Shaders/Public/SubstratePublic.ush
+++ b/Engine/Shaders/Public/SubstratePublic.ush
@@ -78,36 +78,13 @@ FSubstrateGBuffer SubstratePublic_GetSubstrateGBufferFromFirstBSDF(uint2 PixelCo
 		Out.WorldTangent		= SubstrateGetWorldTangent(SubstratePixelHeader, BSDF, SubstrateAddressing);
 		Out.Anisotropy			= SubstrateGetBSDFAnisotropy(BSDF);
 		Out.bIsFirstPerson		= SubstratePixelHeader.IsFirstPerson();
+		Out.Opacity				= SubstrateGetBSDFOpacity(BSDF);
 
-		Out.Opacity			= 0.0f;
+		// Override opacity with SSSHeader value if needed
 		if (SubstrateSubSurfaceHeaderGetSSSType(SSSHeader) != SSS_TYPE_NONE)
 		{
 			Out.Opacity = SubstrateSubSurfaceHeaderGetOpacity(SSSHeader);
 		}
-		else if (BSDF_GETTYPE(BSDF) == SUBSTRATE_BSDF_TYPE_SLAB)
-		{
-			if (BSDF_GETHASFUZZ(BSDF))
-			{
-				Out.Opacity = SLAB_FUZZ_AMOUNT(BSDF);	// Cloth
-			}
-			else if (BSDF_GETHASHAZINESS(BSDF))
-			{
-				const FHaziness Haziness = UnpackHaziness(SLAB_HAZINESS(BSDF));
-				const bool bHazeAsSimpleClearCoat = Haziness.bSimpleClearCoat;
-				if (bHazeAsSimpleClearCoat)
-				{
-					Out.Opacity = Haziness.Weight;		// ClearCoat
-				}
-			}
-		}
-		else if (BSDF_GETTYPE(BSDF) == SUBSTRATE_BSDF_TYPE_EYE)
-		{
-			Out.Opacity = 1 - EYE_IRISMASK(BSDF);		// IrisMask
-		}
-		else if (BSDF_GETTYPE(BSDF) == SUBSTRATE_BSDF_TYPE_HAIR)
-		{
-			Out.Opacity = HAIR_BACKLIT(BSDF);			// BackLit
-		}
 	}
 	else if (SubstratePixelHeader.State == 0)				// Only Unlit and non written pixels have a state of 0. Those pixels should show up as unlit to respect the legacy behavior.
 	{		
