[Shaders] Invalidate D3DCommon.ush GUID to fix preview mode crashes.

#test Several people tested this shelf and confirmed it fixed the crash in QAGame both when built from source and binary from EpicLauncher.
#rnx

diff --git a/Engine/Shaders/Public/Platform/D3D/D3DCommon.ush b/Engine/Shaders/Public/Platform/D3D/D3DCommon.ush
--- a/Engine/Shaders/Public/Platform/D3D/D3DCommon.ush
+++ b/Engine/Shaders/Public/Platform/D3D/D3DCommon.ush
@@ -5,7 +5,7 @@ D3DCommon.ush: Common defines for the D3D shader compiler
 =============================================================================*/
 
 // Update this GUID to invalidate D3D shaders and force a shader recompilation for D3D RHIs.
-#pragma message("UESHADERMETADATA_VERSION 32F71811-58F3-486E-AFC8-F913E12047E6")
+#pragma message("UESHADERMETADATA_VERSION FFA46460-2B9A-4B17-8877-E071C7A31822")
 
 #include "/Engine/Shared/HLSLReservedSpaces.h"
 
