Rebuilding DXC

#Synced-CL 49561730

diff --git a/Engine/Shaders/Public/ShaderVersion.ush b/Engine/Shaders/Public/ShaderVersion.ush
--- a/Engine/Shaders/Public/ShaderVersion.ush
+++ b/Engine/Shaders/Public/ShaderVersion.ush
@@ -3,4 +3,4 @@
 // Each time the console command is executed it generates a new GUID. As this file is included
 // in Platform.ush (which should be included in any shader) it allows to invalidate the shader DDC.
 // 
-#pragma message("UESHADERMETADATA_VERSION 70B48F1C-5F2A-4874-9ED9-C5242BFF44EF")
\ No newline at end of file
+#pragma message("UESHADERMETADATA_VERSION 43D07C7E-F822-4384-8E80-7411A693E609.christopher.waters.4")
\ No newline at end of file
