Substrate - make the default normal not UP to make sure normals are coherent.
Legacy behavior also do not match because it also has a weird default: normal gbuffer is cleared to 0 which means a normal of -1,-1,-1.

Sonormal read + shading model to know if a pixel is valid should be enough and already in place for existing materials.

#tests PC editor
#rnx
#Synced-CL 49615688

diff --git a/Engine/Shaders/Public/SubstratePublic.ush b/Engine/Shaders/Public/SubstratePublic.ush
--- a/Engine/Shaders/Public/SubstratePublic.ush
+++ b/Engine/Shaders/Public/SubstratePublic.ush
@@ -46,7 +46,6 @@ FSubstrateGBuffer SubstratePublic_GetSubstrateGBufferFromFirstBSDF(uint2 PixelCo
 {
 	FSubstrateGBuffer Out = (FSubstrateGBuffer)0; 
 	Out.bIsValid = false;
-	Out.WorldNormal = float3(0.0, 0.0, 1.0);
 
 #if SUBSTRATE_ALLOWS_GBUFFER_SAMPLING
 	#if SUBSTRATE_GBUFFER_FORMAT == 1 && SUBSTRATE_MATERIALCONTAINER_IS_VIEWRESOURCE
