Added #ifndef EARLYDEPTHSTENCIL before defining it as it might have been already defined in PlatformCommon.ush

#rnx

#Synced-CL 49940768

diff --git a/Engine/Shaders/Public/Platform.ush b/Engine/Shaders/Public/Platform.ush
--- a/Engine/Shaders/Public/Platform.ush
+++ b/Engine/Shaders/Public/Platform.ush
@@ -783,9 +783,12 @@ void ShaderYield()
 
 #endif // SM6_PROFILE || SM5_PROFILE || COMPILER_SUPPORTS_ATTRIBUTES
 
+
+#ifndef EARLYDEPTHSTENCIL // Might have been already defined in PlatformCommon.ush
 #if SM6_PROFILE || SM5_PROFILE || METAL_SM5_IOS_TVOS_PROFILE || METAL_SM5_PROFILE || METAL_SM6_PROFILE || ES3_1_PROFILE || VULKAN_PROFILE_SM5 || VULKAN_PROFILE_SM6
 #define EARLYDEPTHSTENCIL [earlydepthstencil]
 #endif
+#endif // EARLYDEPTHSTENCIL
 
 #if COMPILER_SUPPORTS_DUAL_SOURCE_BLENDING_SLOT_DECORATION
 #define DUAL_SOURCE_BLENDING_SLOT(SLOT) [[vk::location(0), vk::index(SLOT)]]
