Updated guid to work around pipeline error on MetalRHI

diff --git a/Engine/Shaders/Public/Platform/Metal/MetalCommon.ush b/Engine/Shaders/Public/Platform/Metal/MetalCommon.ush
--- a/Engine/Shaders/Public/Platform/Metal/MetalCommon.ush
+++ b/Engine/Shaders/Public/Platform/Metal/MetalCommon.ush
@@ -7,7 +7,7 @@
 #pragma once
 
 // Update this GUID to invalidate shader recompilation for only Metal shaders
-#pragma message("UESHADERMETADATA_VERSION 5402461A-0CAA-4A0B-844D-8DE806520E9A")
+#pragma message("UESHADERMETADATA_VERSION 607C4A09-452C-40D7-9C9C-3DBDEFFB50AD")
 
 #if VERTEXSHADER && HULLSHADER
 #ifndef TESSELLATIONSHADER
