[Shaders] Bump GUID in ShaderVersion.ush to fix incremental build due to stale shader files.

#rnx

diff --git a/Engine/Shaders/Public/ShaderVersion.ush b/Engine/Shaders/Public/ShaderVersion.ush
--- a/Engine/Shaders/Public/ShaderVersion.ush
+++ b/Engine/Shaders/Public/ShaderVersion.ush
@@ -3,4 +3,4 @@
 // Each time the console command is executed it generates a new GUID. As this file is included
 // in Platform.ush (which should be included in any shader) it allows to invalidate the shader DDC.
 // 
-#pragma message("UESHADERMETADATA_VERSION 92B142B4-CAEB-47DF-A89D-83910D57E933")
\ No newline at end of file
+#pragma message("UESHADERMETADATA_VERSION 626F9AAB-3832-46E1-8F46-64B302F153D6")
\ No newline at end of file
