Add COMPILER_SUPPORTS_WAVE_MATCH_INTRINSIC to known if WaveMatch is supported or not for a particular platform

#Synced-CL 51289976

diff --git a/Engine/Shaders/Public/Platform.ush b/Engine/Shaders/Public/Platform.ush
--- a/Engine/Shaders/Public/Platform.ush
+++ b/Engine/Shaders/Public/Platform.ush
@@ -135,6 +135,10 @@
 #define COMPILER_SUPPORTS_WAVE_SIZE 0
 #endif
 
+#ifndef COMPILER_SUPPORTS_WAVE_MATCH_INTRINSICS
+#define COMPILER_SUPPORTS_WAVE_MATCH_INTRINSICS 0
+#endif
+
 #ifndef COMPILER_FORCE_WAVE32_MODE
 #define COMPILER_FORCE_WAVE32_MODE
 #endif
diff --git a/Engine/Shaders/Public/Platform/D3D/D3DCommon.ush b/Engine/Shaders/Public/Platform/D3D/D3DCommon.ush
--- a/Engine/Shaders/Public/Platform/D3D/D3DCommon.ush
+++ b/Engine/Shaders/Public/Platform/D3D/D3DCommon.ush
@@ -148,6 +148,10 @@ void UEReportAssertWithPayload(uint AssertID, uint4 Payload)
 	#endif
 #endif
 
+#if __SHADER_TARGET_MAJOR >= 6 && __SHADER_TARGET_MINOR >= 5
+	#define COMPILER_SUPPORTS_WAVE_MATCH_INTRINSICS 1
+#endif
+
 #if PLATFORM_SUPPORTS_BINDLESS
 	#if (RAYCALLABLESHADER || RAYHITGROUPSHADER || RAYGENSHADER || RAYMISSSHADER)
 		// Bindless index can be divergent in hit, miss and callable shaders
