Bump shader version to avoid cook error

#Synced-CL 51223012

diff --git a/Engine/Shaders/Public/ShaderVersion.ush b/Engine/Shaders/Public/ShaderVersion.ush
--- a/Engine/Shaders/Public/ShaderVersion.ush
+++ b/Engine/Shaders/Public/ShaderVersion.ush
@@ -3,4 +3,4 @@
 // Each time the console command is executed it generates a new GUID. As this file is included
 // in Platform.ush (which should be included in any shader) it allows to invalidate the shader DDC.
 // 
-#pragma message("UESHADERMETADATA_VERSION 76323D69-1266-4061-B1C3-291C928683AE")
\ No newline at end of file
+#pragma message("UESHADERMETADATA_VERSION C8A74D67-D68C-462C-B70F-8102DC3D6269")
\ No newline at end of file
