Add WriteIndirectDrawArgs and WriteIndirectIndexedDrawArgs helpers in Platform.ush

#Synced-CL 51456286

diff --git a/Engine/Shaders/Public/Platform.ush b/Engine/Shaders/Public/Platform.ush
--- a/Engine/Shaders/Public/Platform.ush
+++ b/Engine/Shaders/Public/Platform.ush
@@ -1043,9 +1043,25 @@ void WriteDispatchIndirectArgs(RWBuffer<uint> RWIndirectDispatchArgsBuffer, in u
 // sizeof(FRHIDrawIndirectParameters) / sizeof(uint)
 #define DRAW_INDIRECT_UINT_COUNT 4
 
+void WriteIndirectDrawArgs(RWBuffer<uint> OutIndirectArgsBuffer, in uint InIndex, in uint InVertexCount, in uint InInstanceCount, in uint StartVertexLocation=0, in uint StartInstanceLocation=0)
+{
+	OutIndirectArgsBuffer[DRAW_INDIRECT_UINT_COUNT * InIndex + 0] = InVertexCount;
+	OutIndirectArgsBuffer[DRAW_INDIRECT_UINT_COUNT * InIndex + 1] = InInstanceCount;
+	OutIndirectArgsBuffer[DRAW_INDIRECT_UINT_COUNT * InIndex + 2] = StartVertexLocation;
+	OutIndirectArgsBuffer[DRAW_INDIRECT_UINT_COUNT * InIndex + 3] = StartInstanceLocation;
+}
+
 // sizeof(FRHIDrawIndexedIndirectParameters) / sizeof(uint)
 #define DRAW_INDEXED_INDIRECT_UINT_COUNT 5
 
+void WriteIndirectIndexedDrawArgs(RWBuffer<uint> OutIndirectArgsBuffer, in uint InIndex, in uint InIndexCount, in uint InInstanceCount, in uint InStartIndexLocation=0, in uint InBaseVertexLocation=0, in uint InStartInstanceLocation=0)
+{
+	OutIndirectArgsBuffer[DRAW_INDEXED_INDIRECT_UINT_COUNT * InIndex + 0] = InIndexCount;
+	OutIndirectArgsBuffer[DRAW_INDEXED_INDIRECT_UINT_COUNT * InIndex + 1] = InInstanceCount;
+	OutIndirectArgsBuffer[DRAW_INDEXED_INDIRECT_UINT_COUNT * InIndex + 2] = InStartIndexLocation;
+	OutIndirectArgsBuffer[DRAW_INDEXED_INDIRECT_UINT_COUNT * InIndex + 3] = InBaseVertexLocation;
+	OutIndirectArgsBuffer[DRAW_INDEXED_INDIRECT_UINT_COUNT * InIndex + 4] = InStartInstanceLocation;
+}
 
 // ---------------------------------------------------- Compiler missing implementations
 
