Add WaveReadCrossLane

#Synced-CL 52165209

diff --git a/Engine/Shaders/Public/Platform.ush b/Engine/Shaders/Public/Platform.ush
--- a/Engine/Shaders/Public/Platform.ush
+++ b/Engine/Shaders/Public/Platform.ush
@@ -717,6 +717,10 @@
 #define COMPILER_SUPPORTS_WAVE_PERMUTE 0
 #endif
 
+#ifndef COMPILER_SUPPORTS_WAVE_READCROSSLANE
+#define COMPILER_SUPPORTS_WAVE_READCROSSLANE 0
+#endif
+
 // Mirrors GRHISupportsPrimitiveShaders.
 #ifndef COMPILER_SUPPORTS_PRIMITIVE_SHADERS
 #define COMPILER_SUPPORTS_PRIMITIVE_SHADERS 0
