Adding missing descriptor range macros that were missed in a prior checkin.

#Synced-CL 52522322

diff --git a/Engine/Shaders/Public/BindlessResources.ush b/Engine/Shaders/Public/BindlessResources.ush
--- a/Engine/Shaders/Public/BindlessResources.ush
+++ b/Engine/Shaders/Public/BindlessResources.ush
@@ -113,6 +113,12 @@
 
 	#define GetDescriptorRangeSRV(InType, InDescriptorRange, InIndex) GetSRVFromHeap(InType, GetDescriptorRangeStartIndex(InDescriptorRange) + InIndex)
 	#define GetDescriptorRangeUAV(InType, InDescriptorRange, InIndex) GetUAVFromHeap(InType, GetDescriptorRangeStartIndex(InDescriptorRange) + InIndex)
+
+	// Declares the heap array for a given resource type so it can be indexed via GetDescriptorRangeSRV/UAV.
+	// Must be called once per type used with a descriptor range, at global scope in the shader.
+	// On Vulkan this expands to the per-type unbounded array declaration; on other platforms it is a no-op.
+	#define DESCRIPTOR_RANGE_HEAP_DECL_SRV(Type) UE_BINDLESS_SRV_HEAP_DECL(Type)
+	#define DESCRIPTOR_RANGE_HEAP_DECL_UAV(Type) UE_BINDLESS_UAV_HEAP_DECL(Type)
 #endif
 
 //
